Added Rigid Asset Dataflow plugin with initial support for procedurally creating physics asets through dataflow graphs.

Add support for:
- Creating UPhysicsAsset with a new terminal node
- Maipulating bone, constraint, body selections
- Creating geometry for bones using generators
- Sphere generator for body geometry included, more to follow
- Creating joints for bones using generators
- Simple swing/twist generator included, more to follow
- Bypassing the physics asset editor when a dataflow attachment is present, prefering the dataflow editor
- Various object renderers for some new dataflow nodes

#virtualized

diff --git a/Engine/Plugins/Experimental/ChaosRigidAsset/ChaosRigidAsset.uplugin b/Engine/Plugins/Experimental/ChaosRigidAsset/ChaosRigidAsset.uplugin
new file mode 100644
--- /dev/null
+++ b/Engine/Plugins/Experimental/ChaosRigidAsset/ChaosRigidAsset.uplugin
@@ -0,0 +1,42 @@
+{
+	"FileVersion": 3,
+	"Version": 1,
+	"VersionName": "0.1",
+	"FriendlyName": "Chaos Rigid Asset",
+	"Description": "Rigid Asset plugin for creating and utilising collections of rigid bodies",
+	"Category": "Physics",
+	"CreatedBy": "Epic Games, Inc.",
+	"CreatedByURL": "https://epicgames.com",
+	"DocsURL": "",
+	"MarketplaceURL": "",
+	"SupportURL": "",
+	"EnabledByDefault": false,
+	"CanContainContent": true,
+	"IsBetaVersion": false,
+	"IsExperimentalVersion": true,
+	"Installed": false,
+	"Modules": [
+		{
+			"Name": "ChaosRigidAssetEditor",
+			"Type": "Editor",
+			"LoadingPhase": "Default",
+			"PlatformAllowList": [ "Win64", "Mac", "Linux" ]
+		},
+		{
+			"Name": "ChaosRigidAssetNodes",
+			"Type": "Runtime",
+			"LoadingPhase": "Default"
+		},
+		{
+			"Name": "ChaosRigidAssetEngine",
+			"Type": "Runtime",
+			"LoadingPhase": "Default"
+		}
+	],
+	"Plugins": [
+		{
+			"Name": "Dataflow",
+			"Enabled": true
+		}
+	]
+}
diff --git a/Engine/Plugins/Experimental/ChaosRigidAsset/Source/ChaosRigidAssetEditor/ChaosRigidAssetEditor.Build.cs b/Engine/Plugins/Experimental/ChaosRigidAsset/Source/ChaosRigidAssetEditor/ChaosRigidAssetEditor.Build.cs
new file mode 100644
--- /dev/null
+++ b/Engine/Plugins/Experimental/ChaosRigidAsset/Source/ChaosRigidAssetEditor/ChaosRigidAssetEditor.Build.cs
@@ -0,0 +1,40 @@
+// Copyright Epic Games, Inc. All Rights Reserved.
+
+using UnrealBuildTool;
+
+public class ChaosRigidAssetEditor : ModuleRules
+{
+	public ChaosRigidAssetEditor(ReadOnlyTargetRules Target) : base(Target)
+	{
+		PCHUsage = ModuleRules.PCHUsageMode.UseExplicitOrSharedPCHs;
+		
+		PublicDependencyModuleNames.AddRange(
+			new string[]
+			{
+				"Core",
+				"Engine",
+				"UnrealEd",
+				"SlateCore",
+				"Slate",
+				"DataflowCore",
+				"DataflowEngine",
+				"DataflowEditor",
+				"ChaosRigidAssetEngine",
+				"ChaosRigidAssetNodes",
+				"PhysicsAssetEditor"
+			}
+		);
+		
+		
+		PrivateDependencyModuleNames.AddRange(
+			new string[]
+			{
+				"CoreUObject",
+			}
+		);
+
+		bAllowUETypesInNamespaces = true;
+
+		SetupModulePhysicsSupport(Target);
+	}
+}
diff --git a/Engine/Plugins/Experimental/ChaosRigidAsset/Source/ChaosRigidAssetEditor/Private/ChaosRigidAssetEditorModule.cpp b/Engine/Plugins/Experimental/ChaosRigidAsset/Source/ChaosRigidAssetEditor/Private/ChaosRigidAssetEditorModule.cpp
new file mode 100644
--- /dev/null
+++ b/Engine/Plugins/Experimental/ChaosRigidAsset/Source/ChaosRigidAssetEditor/Private/ChaosRigidAssetEditorModule.cpp
@@ -0,0 +1,222 @@
+// Copyright Epic Games, Inc. All Rights Reserved.
+
+#include "ChaosRigidAssetEditorModule.h"
+#include "CoreMinimal.h"
+#include "Modules/ModuleManager.h"
+#include "ContentBrowserModule.h"
+#include "PhysicsEngine/PhysicsAsset.h"
+#include "Dataflow/DataflowEditor.h"
+#include "DataflowAttachment.h"
+#include "DataflowRendering.h"
+#include "Dataflow/DataflowRenderingFactory.h"
+#include "Features/IModularFeatures.h"
+
+IMPLEMENT_MODULE(FChaosRigidAssetEditorModule, ChaosRigidAssetEditor);
+
+bool GRigidAssetModuleOverridePhysicsAssetEditor = false;
+static FAutoConsoleVariableRef GRigidAssetModuleOverridePhysicsAssetEditorCVar(
+	TEXT("p.rigidasset.enableeditoroverride"),
+	GRigidAssetModuleOverridePhysicsAssetEditor,
+	TEXT("Whether to enable the editor override for physics asset such that physics assets using dataflow will no longer open in the base physics asset editor, preferring the dataflow editor instead"));
+
+/**
+ * Override for physics asset editor to instead route to the dataflow editor if there
+ * is a dataflow attachment present in the asset.
+ */
+class FDataflowPhysicsAssetEditorOverride : public IPhysicsAssetEditorOverride
+{
+public:
+
+	virtual ~FDataflowPhysicsAssetEditorOverride() = default;
+
+	bool OpenAsset(UPhysicsAsset* InAsset) override;
+};
+
+UPhysicsAsset* GetPhysAsset(const FAssetData& InAsset)
+{
+	if(UObject* AssetBase = InAsset.GetAsset())
+	{
+		return Cast<UPhysicsAsset>(AssetBase);
+	}
+	return nullptr;
+}
+
+UDataflowAttachment* GetDataflowAttachment(IInterface_AssetUserData* UserDataHolder)
+{
+	if(!UserDataHolder)
+	{
+		return nullptr;
+	}
+
+	return UserDataHolder->GetAssetUserData<UDataflowAttachment>();
+}
+
+bool HasDataflowAttachment(IInterface_AssetUserData* UserDataHolder)
+{
+	if(!UserDataHolder)
+	{
+		return false;
+	}
+
+	return UserDataHolder->HasAssetUserDataOfClass(UDataflowAttachment::StaticClass());
+}
+
+void AddDataflow(FAssetData InAsset)
+{
+	UPhysicsAsset* Asset = GetPhysAsset(InAsset);
+
+	if(!Asset)
+	{
+		return;
+	}
+
+	Asset->AddAssetUserData(NewObject<UDataflowAttachment>(Asset));
+}
+
+void RemoveDataflow(FAssetData InAsset)
+{
+	UPhysicsAsset* Asset = GetPhysAsset(InAsset);
+
+	if(!Asset)
+	{
+		return;
+	}
+
+	Asset->RemoveUserDataOfClass(UDataflowAttachment::StaticClass());
+}
+
+void SpawnEditorFor(UPhysicsAsset* InAsset)
+{
+	UDataflowAttachment* Attachment = GetDataflowAttachment(InAsset);
+	if(Attachment)
+	{
+		if(!Attachment->GetDataflowInstance().GetDataflowAsset())
+		{
+			Attachment->GetDataflowInstance().SetDataflowAsset(Cast<UDataflow>(UE::DataflowAssetDefinitionHelpers::NewOrOpenDataflowAsset(InAsset)));
+		}
+
+		UAssetEditorSubsystem* const AssetEditorSubsystem = GEditor->GetEditorSubsystem<UAssetEditorSubsystem>();
+
+		UDataflowEditor* const AssetEditor = NewObject<UDataflowEditor>(AssetEditorSubsystem, NAME_None, RF_Transient);
+		AssetEditor->RegisterToolCategories({ "General" });
+		const TSubclassOf<AActor> ActorClass = StaticLoadClass(AActor::StaticClass(), nullptr, TEXT("/ChaosRigidAsset/BP_PhysicsAssetPreview.BP_PhysicsAssetPreview_C"), nullptr, LOAD_None, nullptr);
+
+		AssetEditor->Initialize({ GetDataflowAttachment(InAsset) }, ActorClass);
+	}
+}
+
+void EditDataflow(FAssetData InAsset)
+{
+	SpawnEditorFor(GetPhysAsset(InAsset));
+}
+
+bool CanAddDataflow(FAssetData InAsset)
+{
+	UPhysicsAsset* Asset = GetPhysAsset(InAsset);
+	return Asset && !HasDataflowAttachment(Asset);
+}
+
+bool CanRemoveDataflow(FAssetData InAsset)
+{
+	UPhysicsAsset* Asset = GetPhysAsset(InAsset);
+	return Asset && HasDataflowAttachment(Asset);
+}
+
+bool CanEditDataflow(FAssetData InAsset)
+{
+	UPhysicsAsset* Asset = GetPhysAsset(InAsset);
+	return Asset && HasDataflowAttachment(Asset);
+}
+
+void FillMenu(FMenuBuilder& MenuBuilder, const TArray<FAssetData> Assets)
+{
+	if(Assets.Num() == 1 && Assets[0].IsInstanceOf<UPhysicsAsset>())
+	{
+		UPhysicsAsset* PhysAsset = GetPhysAsset(Assets[0]);
+
+		if(!PhysAsset)
+		{
+			return;
+		}
+
+		const bool bHasAttachment = HasDataflowAttachment(PhysAsset);
+
+		MenuBuilder.BeginSection("PhysicsAssetMenu_Dataflow", NSLOCTEXT("RigidAsset", "MenuSection", "DataFlow"));
+
+		// Add Dataflow
+		{
+			FMenuEntryParams Params;
+			Params.DirectActions.ExecuteAction = FExecuteAction::CreateStatic(&AddDataflow, Assets[0]);
+			Params.DirectActions.CanExecuteAction = FCanExecuteAction::CreateStatic(&CanAddDataflow, Assets[0]);
+			Params.LabelOverride = NSLOCTEXT("ChaosRigidAsset", "PhysicsAssetContextAddDataflow", "Add Dataflow");
+			MenuBuilder.AddMenuEntry(Params);
+		}
+
+		// Remove Dataflow
+		{
+			FMenuEntryParams Params;
+			Params.DirectActions.ExecuteAction = FExecuteAction::CreateStatic(&RemoveDataflow, Assets[0]);
+			Params.DirectActions.CanExecuteAction = FCanExecuteAction::CreateStatic(&CanRemoveDataflow, Assets[0]);
+			Params.LabelOverride = NSLOCTEXT("ChaosRigidAsset", "PhysicsAssetContextRemoveDataflow", "Remove Dataflow");
+			MenuBuilder.AddMenuEntry(Params);
+		}
+
+		// Edit Dataflow
+		{
+			FMenuEntryParams Params;
+			Params.DirectActions.ExecuteAction = FExecuteAction::CreateStatic(&EditDataflow, Assets[0]);
+			Params.DirectActions.CanExecuteAction = FCanExecuteAction::CreateStatic(&CanEditDataflow, Assets[0]);
+			Params.LabelOverride = NSLOCTEXT("ChaosRigidAsset", "PhysicsAssetContextEditDataflow", "Edit Dataflow");
+			MenuBuilder.AddMenuEntry(Params);
+		}
+
+		MenuBuilder.EndSection();
+	}
+}
+
+TSharedRef<FExtender> ExtendAssetContextMenu(const TArray<FAssetData>& Assets)
+{
+	TSharedRef<FExtender> Extender = MakeShared<FExtender>();
+
+	Extender->AddMenuExtension("GetAssetActions", EExtensionHook::After, nullptr, FMenuExtensionDelegate::CreateStatic(&FillMenu, Assets));
+
+	return Extender;
+}
+
+void FChaosRigidAssetEditorModule::StartupModule()
+{
+	FContentBrowserModule& ContentBrowser = FModuleManager::LoadModuleChecked<FContentBrowserModule>(TEXT("ContentBrowser"));
+	ContentBrowser.GetAllAssetViewContextMenuExtenders().Add(FContentBrowserMenuExtender_SelectedAssets::CreateStatic(&ExtendAssetContextMenu));
+
+	EditorFeature = MakeUnique<FDataflowPhysicsAssetEditorOverride>();
+
+	IModularFeatures::Get().RegisterModularFeature(IPhysicsAssetEditorOverride::ModularFeatureName, EditorFeature.Get());
+
+	UE::Dataflow::FRenderingFactory::GetInstance()->RegisterCallbacks(MakeUnique<UE::Chaos::RigidAsset::FAggregateGeometryGeomRenderCallbacks>());
+	UE::Dataflow::FRenderingFactory::GetInstance()->RegisterCallbacks(MakeUnique<UE::Chaos::RigidAsset::FBoneSelectionRenderCallbacks>());
+	UE::Dataflow::FRenderingFactory::GetInstance()->RegisterCallbacks(MakeUnique<UE::Chaos::RigidAsset::FPhysAssetStateRenderCallbacks>());
+}
+
+void FChaosRigidAssetEditorModule::ShutdownModule()
+{
+	IModularFeatures::Get().UnregisterModularFeature(IPhysicsAssetEditorOverride::ModularFeatureName, EditorFeature.Get());
+	EditorFeature = nullptr;
+}
+
+bool FDataflowPhysicsAssetEditorOverride::OpenAsset(UPhysicsAsset* InAsset)
+{
+	if(!GRigidAssetModuleOverridePhysicsAssetEditor)
+	{
+		return false;
+	}
+
+	UDataflowAttachment* Attachment = GetDataflowAttachment(InAsset);
+	if(Attachment)
+	{
+		SpawnEditorFor(InAsset);
+
+		return true;
+	}
+
+	return false;
+}
diff --git a/Engine/Plugins/Experimental/ChaosRigidAsset/Source/ChaosRigidAssetEditor/Private/DataflowRendering.cpp b/Engine/Plugins/Experimental/ChaosRigidAsset/Source/ChaosRigidAssetEditor/Private/DataflowRendering.cpp
new file mode 100644
--- /dev/null
+++ b/Engine/Plugins/Experimental/ChaosRigidAsset/Source/ChaosRigidAssetEditor/Private/DataflowRendering.cpp
@@ -0,0 +1,272 @@
+// Copyright Epic Games, Inc. All Rights Reserved.
+
+#include "DataflowRendering.h"
+
+#include "PhysicsEngine/AggregateGeom.h"
+#include "PhysicsEngine/BoxElem.h"
+#include "GeometryCollection/Facades/CollectionRenderingFacade.h"
+
+#include "Engine/SkeletalMesh.h"
+#include "ReferenceSkeleton.h"
+#include "PhysicsEngine/SkeletalBodySetup.h"
+#include "BoneSelection.h"
+#include "PhysicsAssetDataflowState.h"
+
+namespace UE::Chaos::RigidAsset
+{
+	void RenderAggGeom(GeometryCollection::Facades::FRenderingFacade& RenderingFacade, const UE::Dataflow::FGraphRenderingState& State, const FKAggregateGeom& AggGeom, const FTransform& BaseTm, FName BoneName, const FLinearColor& InColor)
+	{
+		const int32 NumElems = AggGeom.GetElementCount();
+
+		if(NumElems == 0)
+		{
+			// Nothing to actually draw - return before opening a geometry group
+			return;
+		}
+
+		for(int32 Index = 0; Index < NumElems; ++Index)
+		{
+			FString GeometryName = FString::Printf(TEXT("%s_%s_%d"), *State.GetGuid().ToString(), *BoneName.ToString(), Index);
+			const int32 GeometryGroup = RenderingFacade.StartGeometryGroup(GeometryName);
+
+			const FKShapeElem* Elem = AggGeom.GetElement(Index);
+
+			if(!Elem)
+			{
+				break;
+			}
+
+			switch(Elem->GetShapeType())
+			{
+			case EAggCollisionShape::Box:
+			{
+				const FKBoxElem* AsBox = static_cast<const FKBoxElem*>(Elem);
+				const FVector3f Extent = FVector3f{ AsBox->X, AsBox->Y, AsBox->Z } *0.5f;
+
+				RenderingFacade.SetGroupTransform(GeometryGroup, BaseTm * Elem->GetTransform());
+				RenderingFacade.AddBox({ -Extent, Extent });
+
+				break;
+			}
+			case EAggCollisionShape::Sphere:
+			{
+				const FKSphereElem* AsSphere = static_cast<const FKSphereElem*>(Elem);
+
+				RenderingFacade.SetGroupTransform(GeometryGroup, BaseTm * Elem->GetTransform());
+				RenderingFacade.AddSphere((FVector3f)AsSphere->Center, AsSphere->Radius, InColor);
+
+				break;
+			}
+			case EAggCollisionShape::Sphyl:
+			{
+				const FKSphylElem* AsSphyl = static_cast<const FKSphylElem*>(Elem);
+
+				// AddCapsule uses the capsule mesh generator which has the origin of the capsule at one end of
+				// the segment, we need to transform first to the centre of the capsule, then apply the elem transform
+				const FTransform CapsuleCollisionAdjust(FQuat::Identity, FVector(0.0f, 0.0f, -AsSphyl->Length / 2));
+
+				RenderingFacade.SetGroupTransform(GeometryGroup, BaseTm * CapsuleCollisionAdjust * Elem->GetTransform());
+				RenderingFacade.AddCapsule(AsSphyl->Length, AsSphyl->Radius);
+
+				break;
+			}
+			default:
+				break;
+			}
+
+			RenderingFacade.EndGeometryGroup(GeometryGroup);
+		}
+	}
+
+	UE::Dataflow::FRenderKey FAggregateGeometryGeomRenderCallbacks::GetRenderKey() const
+	{
+		return { "GeomRender", FName("FKAggregateGeom") };
+	}
+
+	bool FAggregateGeometryGeomRenderCallbacks::CanRender(const UE::Dataflow::IDataflowConstructionViewMode& ViewMode) const
+	{
+		return true;
+	}
+
+	void FAggregateGeometryGeomRenderCallbacks::Render(GeometryCollection::Facades::FRenderingFacade& RenderingFacade, const UE::Dataflow::FGraphRenderingState& State)
+	{
+		if(State.GetRenderOutputs().Num() > 0)
+		{
+
+			FKAggregateGeom Default;
+			const FKAggregateGeom& AggGeom = State.GetValue<FKAggregateGeom>(State.GetRenderOutputs()[0], Default);
+
+			RenderAggGeom(RenderingFacade, State, AggGeom, FTransform::Identity, NAME_None, FLinearColor::White);
+		}
+	}
+
+	UE::Dataflow::FRenderKey FBoneSelectionRenderCallbacks::GetRenderKey() const
+	{
+		return { "GeomRender", FName("FBoneSelection") };
+	}
+
+	bool FBoneSelectionRenderCallbacks::CanRender(const UE::Dataflow::IDataflowConstructionViewMode& ViewMode) const
+	{
+		return true;
+	}
+
+	// GeometryCore perp vectors impl, not currently exposed in that module
+	template <typename RealType>
+	void MakePerpVectors(const UE::Math::TVector<RealType>& Normal, UE::Math::TVector<RealType>& OutPerp1, UE::Math::TVector<RealType>& OutPerp2)
+	{
+		// Duff et al method, from https://graphics.pixar.com/library/OrthonormalB/paper.pdf
+		if(Normal.Z < (RealType)0)
+		{
+			RealType A = (RealType)1 / ((RealType)1 - Normal.Z);
+			RealType B = Normal.X * Normal.Y * A;
+			OutPerp1.X = (RealType)1 - Normal.X * Normal.X * A;
+			OutPerp1.Y = -B;
+			OutPerp1.Z = Normal.X;
+			OutPerp2.X = B;
+			OutPerp2.Y = Normal.Y * Normal.Y * A - (RealType)1;
+			OutPerp2.Z = -Normal.Y;
+		}
+		else
+		{
+			RealType A = (RealType)1 / ((RealType)1 + Normal.Z);
+			RealType B = -Normal.X * Normal.Y * A;
+			OutPerp1.X = (RealType)1 - Normal.X * Normal.X * A;
+			OutPerp1.Y = B;
+			OutPerp1.Z = -Normal.X;
+			OutPerp2.X = B;
+			OutPerp2.Y = (RealType)1 - Normal.Y * Normal.Y * A;
+			OutPerp2.Z = -Normal.Y;
+		}
+	}
+
+	void DrawBone(GeometryCollection::Facades::FRenderingFacade& RenderingFacade, const FVector3f& From, const FVector3f& To, const FLinearColor& InColor)
+	{
+		float Length;
+		FVector3f Axes[3];
+		const FVector3f Segment = To - From;
+		Segment.ToDirectionAndLength(Axes[2], Length);
+		MakePerpVectors(Axes[2], Axes[0], Axes[1]);
+
+		constexpr float WidePoint = 0.3f;
+		constexpr float Thickness = 2.5f;
+
+		TArray<FVector3f> Verts;
+		TArray<FLinearColor> Colors;
+		Verts.SetNumUninitialized(6);
+		Colors.SetNumUninitialized(6);
+
+		Verts[0] = From;
+		Verts[5] = To;
+
+		Verts[1] = From + (Axes[0] + Axes[1]) * Thickness + Axes[2] * (WidePoint * Length);
+		Verts[2] = From + (Axes[0] + -Axes[1]) * Thickness + Axes[2] * (WidePoint * Length);
+		Verts[3] = From + (-Axes[0] + -Axes[1]) * Thickness + Axes[2] * (WidePoint * Length);
+		Verts[4] = From + (-Axes[0] + Axes[1]) * Thickness + Axes[2] * (WidePoint * Length);
+
+		TArray<FIntVector> Indices;
+		Indices.SetNumUninitialized(8);
+
+		// Top Section
+		Indices[0] = FIntVector{ 0, 2, 1 };
+		Indices[1] = FIntVector{ 0, 3, 2 };
+		Indices[2] = FIntVector{ 0, 4, 3 };
+		Indices[3] = FIntVector{ 0, 1, 4 };
+
+		// Bottom Section
+		Indices[4] = FIntVector{ 5, 1, 2 };
+		Indices[5] = FIntVector{ 5, 2, 3 };
+		Indices[6] = FIntVector{ 5, 3, 4 };
+		Indices[7] = FIntVector{ 5, 4, 1 };
+
+		for(FLinearColor& Color : Colors)
+		{
+			Color = InColor;
+		}
+
+		RenderingFacade.AddFaces(Verts, Indices, Colors);
+	}
+
+	void FBoneSelectionRenderCallbacks::Render(GeometryCollection::Facades::FRenderingFacade& RenderingFacade, const UE::Dataflow::FGraphRenderingState& State)
+	{
+		if(State.GetRenderOutputs().Num() > 0)
+		{
+			FRigidAssetBoneSelection Default;
+			const FRigidAssetBoneSelection& Selection = State.GetValue<FRigidAssetBoneSelection>(State.GetRenderOutputs()[0], Default);
+
+			if(!Selection.Mesh || !Selection.Skeleton)
+			{
+				return;
+			}
+
+			const FReferenceSkeleton& RefSkel = Selection.Mesh->GetRefSkeleton();
+
+			TArray<FTransform> BoneTransforms;
+			TArray<int32> ParentIndices;
+			TArray<bool> SelectedIndices;
+
+			RefSkel.GetBoneAbsoluteTransforms(BoneTransforms);
+
+			const int32 NumBones = BoneTransforms.Num();
+
+			SelectedIndices.SetNumZeroed(NumBones);
+			for(const FRigidAssetBoneInfo& Selected : Selection.SelectedBones)
+			{
+				SelectedIndices[Selected.Index] = true;
+			}
+
+			ParentIndices.SetNum(NumBones);
+			for(int32 BoneIndex = 0; BoneIndex < NumBones; ++BoneIndex)
+			{
+				ParentIndices[BoneIndex] = RefSkel.GetParentIndex(BoneIndex);
+			}
+
+			int32 Group = RenderingFacade.StartGeometryGroup(TEXT("BoneSelection"));
+			for(int32 BoneIndex = 0; BoneIndex < NumBones; ++BoneIndex)
+			{
+				const FTransform& BoneTransform = BoneTransforms[BoneIndex];
+
+				RenderingFacade.AddPoint((FVector3f)BoneTransform.GetTranslation());
+
+				const int32 ParentIndex = ParentIndices[BoneIndex];
+				if(ParentIndex != INDEX_NONE)
+				{
+					const FLinearColor BoneColor = SelectedIndices[ParentIndex] ? FLinearColor(1.0f, 0.6f, 0.0f) : FLinearColor::Gray;
+					DrawBone(RenderingFacade, (FVector3f)BoneTransforms[ParentIndex].GetTranslation(), (FVector3f)BoneTransforms[BoneIndex].GetTranslation(), BoneColor);
+				}
+			}
+			RenderingFacade.EndGeometryGroup(Group);
+		}
+	}
+
+	UE::Dataflow::FRenderKey FPhysAssetStateRenderCallbacks::GetRenderKey() const
+	{
+		return { "GeomRender", FName("FPhysicsAssetDataflowState") };
+	}
+
+	bool FPhysAssetStateRenderCallbacks::CanRender(const UE::Dataflow::IDataflowConstructionViewMode& ViewMode) const
+	{
+		return true;
+	}
+
+	void FPhysAssetStateRenderCallbacks::Render(GeometryCollection::Facades::FRenderingFacade& RenderingFacade, const UE::Dataflow::FGraphRenderingState& State)
+	{
+		if(State.GetRenderOutputs().Num() > 0)
+		{
+			FPhysicsAssetDataflowState Default;
+			const FPhysicsAssetDataflowState& AssetState = State.GetValue<FPhysicsAssetDataflowState>(State.GetRenderOutputs()[0], Default);
+
+			if(AssetState.HasData())
+			{
+				const FReferenceSkeleton& RefSkeleton = AssetState.TargetMesh->GetRefSkeleton();
+
+				for(TObjectPtr<USkeletalBodySetup> Body : AssetState.Bodies)
+				{
+					const int32 BoneIndex = RefSkeleton.FindBoneIndex(Body->BoneName);
+					FTransform BoneTransform = RefSkeleton.GetBoneAbsoluteTransform(BoneIndex);
+
+					RenderAggGeom(RenderingFacade, State, Body->AggGeom, BoneTransform, Body->BoneName, FLinearColor(1.0f, 1.0f, 1.0f, 0.5f));
+				}
+			}
+		}
+	}
+}
diff --git a/Engine/Plugins/Experimental/ChaosRigidAsset/Source/ChaosRigidAssetEditor/Public/ChaosRigidAssetEditorModule.h b/Engine/Plugins/Experimental/ChaosRigidAsset/Source/ChaosRigidAssetEditor/Public/ChaosRigidAssetEditorModule.h
new file mode 100644
--- /dev/null
+++ b/Engine/Plugins/Experimental/ChaosRigidAsset/Source/ChaosRigidAssetEditor/Public/ChaosRigidAssetEditorModule.h
@@ -0,0 +1,19 @@
+// Copyright Epic Games, Inc. All Rights Reserved.
+
+#pragma once
+
+#include "Modules/ModuleInterface.h"
+#include "IPhysicsAssetEditor.h"
+
+class FDataflowPhysicsAssetEditorOverride;
+
+class FChaosRigidAssetEditorModule : public IModuleInterface
+{
+public:
+
+	CHAOSRIGIDASSETEDITOR_API virtual void StartupModule() override;
+	CHAOSRIGIDASSETEDITOR_API virtual void ShutdownModule() override;
+
+private:
+	TUniquePtr<FDataflowPhysicsAssetEditorOverride> EditorFeature;
+};
diff --git a/Engine/Plugins/Experimental/ChaosRigidAsset/Source/ChaosRigidAssetEditor/Public/DataflowRendering.h b/Engine/Plugins/Experimental/ChaosRigidAsset/Source/ChaosRigidAssetEditor/Public/DataflowRendering.h
new file mode 100644
--- /dev/null
+++ b/Engine/Plugins/Experimental/ChaosRigidAsset/Source/ChaosRigidAssetEditor/Public/DataflowRendering.h
@@ -0,0 +1,36 @@
+// Copyright Epic Games, Inc. All Rights Reserved.
+
+#pragma once
+
+#include "Dataflow/DataflowRenderingFactory.h"
+
+namespace UE::Chaos::RigidAsset
+{
+	class FAggregateGeometryGeomRenderCallbacks : public UE::Dataflow::FRenderingFactory::ICallbackInterface
+	{
+	public:
+
+		UE::Dataflow::FRenderKey GetRenderKey() const override;
+		bool CanRender(const UE::Dataflow::IDataflowConstructionViewMode& ViewMode) const override;
+		void Render(GeometryCollection::Facades::FRenderingFacade& RenderingFacade, const UE::Dataflow::FGraphRenderingState& State) override;
+
+	};
+
+	class FBoneSelectionRenderCallbacks : public UE::Dataflow::FRenderingFactory::ICallbackInterface
+	{
+	public:
+
+		UE::Dataflow::FRenderKey GetRenderKey() const override;
+		bool CanRender(const UE::Dataflow::IDataflowConstructionViewMode& ViewMode) const override;
+		void Render(GeometryCollection::Facades::FRenderingFacade& RenderingFacade, const UE::Dataflow::FGraphRenderingState& State) override;
+	};
+
+	class FPhysAssetStateRenderCallbacks : public UE::Dataflow::FRenderingFactory::ICallbackInterface
+	{
+	public:
+
+		UE::Dataflow::FRenderKey GetRenderKey() const override;
+		bool CanRender(const UE::Dataflow::IDataflowConstructionViewMode& ViewMode) const override;
+		void Render(GeometryCollection::Facades::FRenderingFacade& RenderingFacade, const UE::Dataflow::FGraphRenderingState& State) override;
+	};
+}
diff --git a/Engine/Plugins/Experimental/ChaosRigidAsset/Source/ChaosRigidAssetEngine/ChaosRigidAssetEngine.Build.cs b/Engine/Plugins/Experimental/ChaosRigidAsset/Source/ChaosRigidAssetEngine/ChaosRigidAssetEngine.Build.cs
new file mode 100644
--- /dev/null
+++ b/Engine/Plugins/Experimental/ChaosRigidAsset/Source/ChaosRigidAssetEngine/ChaosRigidAssetEngine.Build.cs
@@ -0,0 +1,31 @@
+// Copyright Epic Games, Inc. All Rights Reserved.
+
+using UnrealBuildTool;
+
+public class ChaosRigidAssetEngine : ModuleRules
+{
+	public ChaosRigidAssetEngine(ReadOnlyTargetRules Target) : base(Target)
+	{
+		PCHUsage = ModuleRules.PCHUsageMode.UseExplicitOrSharedPCHs;
+		
+		PublicDependencyModuleNames.AddRange(
+			new string[]
+			{
+				"Core",
+				"Engine",
+				"DataflowCore",
+				"DataflowEngine",
+			}
+		);
+		
+		
+		PrivateDependencyModuleNames.AddRange(
+			new string[]
+			{
+				"CoreUObject",
+			}
+		);
+
+		bAllowUETypesInNamespaces = true;
+	}
+}
diff --git a/Engine/Plugins/Experimental/ChaosRigidAsset/Source/ChaosRigidAssetEngine/Private/ChaosRigidAssetEngineModule.cpp b/Engine/Plugins/Experimental/ChaosRigidAsset/Source/ChaosRigidAssetEngine/Private/ChaosRigidAssetEngineModule.cpp
new file mode 100644
--- /dev/null
+++ b/Engine/Plugins/Experimental/ChaosRigidAsset/Source/ChaosRigidAssetEngine/Private/ChaosRigidAssetEngineModule.cpp
@@ -0,0 +1,18 @@
+// Copyright Epic Games, Inc. All Rights Reserved.
+
+#include "Modules/ModuleInterface.h"
+#include "Modules/ModuleManager.h"
+
+class FChaosRigidAssetEngineModule : public IModuleInterface
+{
+public:
+	virtual void StartupModule() override
+	{
+	}
+
+	virtual void ShutdownModule() override
+	{
+	}
+};
+
+IMPLEMENT_MODULE(FChaosRigidAssetEngineModule, ChaosRigidAssetEngine);
diff --git a/Engine/Plugins/Experimental/ChaosRigidAsset/Source/ChaosRigidAssetEngine/Private/DataflowAttachment.cpp b/Engine/Plugins/Experimental/ChaosRigidAsset/Source/ChaosRigidAssetEngine/Private/DataflowAttachment.cpp
new file mode 100644
--- /dev/null
+++ b/Engine/Plugins/Experimental/ChaosRigidAsset/Source/ChaosRigidAssetEngine/Private/DataflowAttachment.cpp
@@ -0,0 +1,47 @@
+// Copyright Epic Games, Inc. All Rights Reserved.
+
+#include "DataflowAttachment.h"
+#include "PhysicsEngine/PhysicsAsset.h"
+#include "PhysicsAssetDataflowContent.h"
+
+void UDataflowAttachment::WriteDataflowContent(const TObjectPtr<UDataflowBaseContent>& DataflowContent) const
+{
+	if(const TObjectPtr<UPhysicsAssetDataflowContent> Content = Cast<UPhysicsAssetDataflowContent>(DataflowContent))
+	{
+		Content->SetDataflowAsset(Instance.GetDataflowAsset());
+		Content->SetDataflowTerminal(Instance.GetDataflowTerminal().ToString());
+
+		if(UPhysicsAsset* PhysAsset = GetTypedOuter<UPhysicsAsset>())
+		{
+			Content->SetSkeletalMesh(PhysAsset->GetPreviewMesh());
+			Content->SetPhysicsAsset(PhysAsset);
+		}
+	}
+}
+
+void UDataflowAttachment::ReadDataflowContent(const TObjectPtr<UDataflowBaseContent>& DataflowContent)
+{
+	
+}
+
+const FDataflowInstance& UDataflowAttachment::GetDataflowInstance() const
+{
+	return Instance;
+}
+
+FDataflowInstance& UDataflowAttachment::GetDataflowInstance()
+{
+	return Instance;
+}
+
+TObjectPtr<UDataflowBaseContent> UDataflowAttachment::CreateDataflowContent()
+{
+	TObjectPtr<UDataflowBaseContent> BaseContent = NewObject<UPhysicsAssetDataflowContent>();
+
+	BaseContent->SetDataflowOwner(this);
+	BaseContent->SetTerminalAsset(this);
+
+	WriteDataflowContent(BaseContent);
+
+	return BaseContent;
+}
diff --git a/Engine/Plugins/Experimental/ChaosRigidAsset/Source/ChaosRigidAssetEngine/Private/PhysicsAssetDataflowContent.cpp b/Engine/Plugins/Experimental/ChaosRigidAsset/Source/ChaosRigidAssetEngine/Private/PhysicsAssetDataflowContent.cpp
new file mode 100644
--- /dev/null
+++ b/Engine/Plugins/Experimental/ChaosRigidAsset/Source/ChaosRigidAssetEngine/Private/PhysicsAssetDataflowContent.cpp
@@ -0,0 +1,22 @@
+// Copyright Epic Games, Inc. All Rights Reserved.
+
+#include "PhysicsAssetDataflowContent.h"
+#include "Engine/SkeletalMesh.h"
+#include "PhysicsEngine/PhysicsAsset.h"
+
+void UPhysicsAssetDataflowContent::SetActorProperties(TObjectPtr<AActor>& PreviewActor) const
+{
+	Super::SetActorProperties(PreviewActor);
+	OverrideActorProperty(PreviewActor, SkelMesh, TEXT("PreviewMesh"));
+	OverrideActorProperty(PreviewActor, PhysAsset, TEXT("PhysicsAsset"));
+}
+
+void UPhysicsAssetDataflowContent::SetSkeletalMesh(USkeletalMesh* InMesh)
+{
+	SkelMesh = InMesh;
+}
+
+void UPhysicsAssetDataflowContent::SetPhysicsAsset(UPhysicsAsset* InAsset)
+{
+	PhysAsset = InAsset;
+}
diff --git a/Engine/Plugins/Experimental/ChaosRigidAsset/Source/ChaosRigidAssetEngine/Private/PhysicsAssetDataflowContent.h b/Engine/Plugins/Experimental/ChaosRigidAsset/Source/ChaosRigidAssetEngine/Private/PhysicsAssetDataflowContent.h
new file mode 100644
--- /dev/null
+++ b/Engine/Plugins/Experimental/ChaosRigidAsset/Source/ChaosRigidAssetEngine/Private/PhysicsAssetDataflowContent.h
@@ -0,0 +1,39 @@
+// Copyright Epic Games, Inc. All Rights Reserved.
+
+#pragma once
+
+#include "UObject/ObjectPtr.h"
+#include "Dataflow/DataflowContent.h"
+
+#include "PhysicsAssetDataflowContent.generated.h"
+
+class AActor;
+class USkeletalMesh;
+class UPhysicsAsset;
+
+/**
+ * Dataflow editor content for physics assets
+ */
+UCLASS()
+class UPhysicsAssetDataflowContent : public UDataflowBaseContent
+{
+	GENERATED_BODY()
+
+public:
+
+	void SetActorProperties(TObjectPtr<AActor>& PreviewActor) const override;
+
+	void SetSkeletalMesh(USkeletalMesh* InMesh);
+	void SetPhysicsAsset(UPhysicsAsset* InAsset);
+
+private:
+	
+	// Preview mesh to use
+	UPROPERTY()
+	TObjectPtr<USkeletalMesh> SkelMesh = nullptr;
+
+	// Physics asset to use (must be compatible with the preview mesh)
+	UPROPERTY()
+	TObjectPtr<UPhysicsAsset> PhysAsset = nullptr;
+
+};
\ No newline at end of file
diff --git a/Engine/Plugins/Experimental/ChaosRigidAsset/Source/ChaosRigidAssetEngine/Public/DataflowAttachment.h b/Engine/Plugins/Experimental/ChaosRigidAsset/Source/ChaosRigidAssetEngine/Public/DataflowAttachment.h
new file mode 100644
--- /dev/null
+++ b/Engine/Plugins/Experimental/ChaosRigidAsset/Source/ChaosRigidAssetEngine/Public/DataflowAttachment.h
@@ -0,0 +1,47 @@
+// Copyright Epic Games, Inc. All Rights Reserved.
+
+#pragma once
+
+#include "Engine/AssetUserData.h"
+#include "Dataflow/DataflowContent.h"
+#include "Dataflow/DataflowInstance.h"
+#include "UObject/ObjectPtr.h"
+
+#include "DataflowAttachment.generated.h"
+
+/**
+ * Generic dataflow instance attachment implemented as a user asset data object so it can be embedded inside
+ * any asset that is a user-data container. This allows assets in modules that do not, or cannot know about
+ * dataflow (e.g. Engine module assets) to contain a dataflow attachment.
+ * Using a modular feature implementation this data can be grabbed as an asset opens and a plugin specific
+ * implementation can be used for the asset editor (see FDataflowPhysicsAssetEditorOverride)
+ */
+UCLASS()
+class UDataflowAttachment : public UAssetUserData, public IDataflowContentOwner, public IDataflowInstanceInterface
+{
+	GENERATED_BODY()
+
+public:
+
+	// IDataflowContentOwner
+	void WriteDataflowContent(const TObjectPtr<UDataflowBaseContent>& DataflowContent) const override;
+	void ReadDataflowContent(const TObjectPtr<UDataflowBaseContent>& DataflowContent) override;
+	////////////////////////////////////////////////////////////////////////////////////////////////////////////////
+	
+	// IDataflowInstanceInterface
+	const FDataflowInstance& GetDataflowInstance() const override;
+	FDataflowInstance& GetDataflowInstance() override;
+	////////////////////////////////////////////////////////////////////////////////////////////////////////////////
+
+private:
+
+	// Dataflow instance for the attachment when editing.
+	UPROPERTY()
+	FDataflowInstance Instance;
+
+protected:
+	
+	// IDataflowContentOwner
+	TObjectPtr<UDataflowBaseContent> CreateDataflowContent() override;
+	////////////////////////////////////////////////////////////////////////////////////////////////////////////////
+};
\ No newline at end of file
diff --git a/Engine/Plugins/Experimental/ChaosRigidAsset/Source/ChaosRigidAssetNodes/ChaosRigidAssetNodes.Build.cs b/Engine/Plugins/Experimental/ChaosRigidAsset/Source/ChaosRigidAssetNodes/ChaosRigidAssetNodes.Build.cs
new file mode 100644
--- /dev/null
+++ b/Engine/Plugins/Experimental/ChaosRigidAsset/Source/ChaosRigidAssetNodes/ChaosRigidAssetNodes.Build.cs
@@ -0,0 +1,34 @@
+// Copyright Epic Games, Inc. All Rights Reserved.
+
+using UnrealBuildTool;
+
+public class ChaosRigidAssetNodes : ModuleRules
+{
+	public ChaosRigidAssetNodes(ReadOnlyTargetRules Target) : base(Target)
+	{
+		PCHUsage = ModuleRules.PCHUsageMode.UseExplicitOrSharedPCHs;
+		
+		PublicDependencyModuleNames.AddRange(
+			new string[]
+			{
+				"Core",
+				"Engine",
+				"DataflowCore",
+				"DataflowEngine",
+				"ChaosRigidAssetEngine"
+			}
+		);
+		
+		
+		PrivateDependencyModuleNames.AddRange(
+			new string[]
+			{
+				"CoreUObject",
+			}
+		);
+
+		SetupModulePhysicsSupport(Target);
+
+		bAllowUETypesInNamespaces = true;
+	}
+}
diff --git a/Engine/Plugins/Experimental/ChaosRigidAsset/Source/ChaosRigidAssetNodes/Internal/PhysicsAssetDataflowState.h b/Engine/Plugins/Experimental/ChaosRigidAsset/Source/ChaosRigidAssetNodes/Internal/PhysicsAssetDataflowState.h
new file mode 100644
--- /dev/null
+++ b/Engine/Plugins/Experimental/ChaosRigidAsset/Source/ChaosRigidAssetNodes/Internal/PhysicsAssetDataflowState.h
@@ -0,0 +1,56 @@
+// Copyright Epic Games, Inc. All Rights Reserved.
+
+#pragma once
+
+#include "CopyOnWriteArray.h"
+
+#include "PhysicsAssetDataflowState.generated.h"
+
+class USkeleton;
+class USkeletalMesh;
+class USkeletalBodySetup;
+class UPhysicsConstraintTemplate;
+
+class FReferenceCollector;
+
+/**
+ * State container for physics asset data during dataflow graph evaluation.
+ * Represents the complete state required to build a physics asset. Datflow nodes populate this container
+ * until it reaches a terminal which packages it into a final UPhysicsAsset. This allows us to keep the
+ * const invariant for dataflow evaluation by adding copy-on-write support to the state.
+ * During node evaluation this object should be passed by-value as it is a collection of pts to the 
+ * actual data. If the copy-on-write arrays are modified they will automatically duplicate if required
+ */
+USTRUCT()
+struct FPhysicsAssetDataflowState
+{
+	GENERATED_BODY()
+
+	void AddStructReferencedObjects(FReferenceCollector& Collector);
+
+	// Base asset data
+	TObjectPtr<USkeleton> TargetSkeleton = nullptr;
+	TObjectPtr<USkeletalMesh> TargetMesh = nullptr;
+
+	// Bodies and constraints in the simulation setup
+	TCopyOnWriteArray<TObjectPtr<USkeletalBodySetup>> Bodies;
+	TCopyOnWriteArray<TObjectPtr<UPhysicsConstraintTemplate>> Constraints;
+
+	void DebugLog() const;
+	CHAOSRIGIDASSETNODES_API bool HasData() const;
+
+private:
+
+};
+
+/**
+ * The state struct has UObject references within the copy-on-write arrays, enable handling for that here
+ */
+template<>
+struct TStructOpsTypeTraits<FPhysicsAssetDataflowState> : public TStructOpsTypeTraitsBase2<FPhysicsAssetDataflowState>
+{
+	enum
+	{
+		WithAddStructReferencedObjects = true
+	};
+};
\ No newline at end of file
diff --git a/Engine/Plugins/Experimental/ChaosRigidAsset/Source/ChaosRigidAssetNodes/Private/BoneSelection.cpp b/Engine/Plugins/Experimental/ChaosRigidAsset/Source/ChaosRigidAssetNodes/Private/BoneSelection.cpp
new file mode 100644
--- /dev/null
+++ b/Engine/Plugins/Experimental/ChaosRigidAsset/Source/ChaosRigidAssetNodes/Private/BoneSelection.cpp
@@ -0,0 +1,40 @@
+// Copyright Epic Games, Inc. All Rights Reserved.
+
+#include "BoneSelection.h"
+
+FRigidAssetBoneInfo::FRigidAssetBoneInfo(FName InName, int32 InIndex, int32 InDepth) : Name(InName)
+, Index(InIndex)
+, Depth(InDepth)
+{
+
+}
+
+bool FRigidAssetBoneInfo::operator<(const FRigidAssetBoneInfo& Other) const
+{
+	return Depth < Other.Depth;
+}
+
+bool FRigidAssetBoneSelection::ContainsIndex(int32 InBoneIndex) const
+{
+	const auto FindByIndex = [&InBoneIndex](const FRigidAssetBoneInfo& Item)
+		{
+			return Item.Index == InBoneIndex;
+		};
+
+	return SelectedBones.FindByPredicate(FindByIndex) != nullptr;
+}
+
+bool FRigidAssetBoneSelection::Contains(FName InBoneName) const
+{
+	const auto FindByName = [&InBoneName](const FRigidAssetBoneInfo& Item)
+		{
+			return Item.Name == InBoneName;
+		};
+
+	return SelectedBones.FindByPredicate(FindByName) != nullptr;
+}
+
+void FRigidAssetBoneSelection::SortBones()
+{
+	SelectedBones.Sort();
+}
diff --git a/Engine/Plugins/Experimental/ChaosRigidAsset/Source/ChaosRigidAssetNodes/Private/ChaosRigidAssetNodesModule.cpp b/Engine/Plugins/Experimental/ChaosRigidAsset/Source/ChaosRigidAssetNodes/Private/ChaosRigidAssetNodesModule.cpp
new file mode 100644
--- /dev/null
+++ b/Engine/Plugins/Experimental/ChaosRigidAsset/Source/ChaosRigidAssetNodes/Private/ChaosRigidAssetNodesModule.cpp
@@ -0,0 +1,19 @@
+// Copyright Epic Games, Inc. All Rights Reserved.
+
+#include "ChaosRigidAssetNodesModule.h"
+#include "ShapeElemNodes.h"
+#include "PhysicsAssetDataflowNodes.h"
+
+IMPLEMENT_MODULE(FChaosRigidAssetNodesModule, ChaosRigidAssetNodes);
+
+ void FChaosRigidAssetNodesModule::StartupModule()
+{
+	UE::Dataflow::RegisterPhysicsAssetTerminalNode();
+	UE::Dataflow::RegisterPhysicsAssetNodes();
+	UE::Dataflow::RegisterShapeNodes();
+}
+
+ void FChaosRigidAssetNodesModule::ShutdownModule()
+{
+
+}
diff --git a/Engine/Plugins/Experimental/ChaosRigidAsset/Source/ChaosRigidAssetNodes/Private/Generators/BoneGeometryGenerators.cpp b/Engine/Plugins/Experimental/ChaosRigidAsset/Source/ChaosRigidAssetNodes/Private/Generators/BoneGeometryGenerators.cpp
new file mode 100644
--- /dev/null
+++ b/Engine/Plugins/Experimental/ChaosRigidAsset/Source/ChaosRigidAssetNodes/Private/Generators/BoneGeometryGenerators.cpp
@@ -0,0 +1,283 @@
+// Copyright Epic Games, Inc. All Rights Reserved.
+
+#include "Generators/BoneGeometryGenerators.h"
+#include "Dataflow/DataflowNode.h"
+#include "Dataflow/DataflowNodeFactory.h"
+#include "Rendering/SkeletalMeshModel.h"
+#include "Engine/SkeletalMesh.h"
+#include "ReferenceSkeleton.h"
+#include "Rendering/SkeletalMeshRenderData.h"
+#include "Rendering/SkeletalMeshLODRenderData.h"
+#include "Rendering/SkinWeightVertexBuffer.h"
+#include "Rendering/PositionVertexBuffer.h"
+#include "Rendering/StaticMeshVertexBuffer.h"
+
+namespace UE::Dataflow
+{
+	void RegisterBoneGeometryGeneratorNodes()
+	{
+		// Body Generators
+		DATAFLOW_NODE_REGISTER_CREATION_FACTORY(FMakeBoxBoneGeometryGenerator);
+		DATAFLOW_NODE_REGISTER_CREATION_FACTORY(FMakeSphereBoneGeometryGenerator);
+		DATAFLOW_NODE_REGISTER_CREATION_FACTORY(FMakeCapsuleBoneGeometryGenerator);
+		DATAFLOW_NODE_REGISTER_CREATION_FACTORY(FMakeConvexBoneGeometryGenerator);
+	}
+}
+
+////////////////////////////////////////////////////////////////////////////////////////////////////////////////
+
+void FMakeBoxBoneGeometryGenerator::Evaluate(UE::Dataflow::FContext& Context, const FDataflowOutput* Out) const
+{
+	if(Out->IsA(&Generator))
+	{
+		TObjectPtr<UBoneGeometryGenerator> OutGenerator = CastChecked<UBoneGeometryGenerator>(StaticDuplicateObject(Generator, GetTransientPackage()));
+		SetValue(Context, OutGenerator, &Generator);
+	}
+}
+
+void FMakeBoxBoneGeometryGenerator::Register()
+{
+	AddOutputs(Generator);
+
+	Generator = NewObject<UBoneGeometryGenerator_Box>(GetOwner());
+}
+
+////////////////////////////////////////////////////////////////////////////////////////////////////////////////
+
+void FMakeSphereBoneGeometryGenerator::Evaluate(UE::Dataflow::FContext& Context, const FDataflowOutput* Out) const
+{
+	if(Out->IsA(&Generator))
+	{
+		TObjectPtr<UBoneGeometryGenerator> OutGenerator = CastChecked<UBoneGeometryGenerator>(StaticDuplicateObject(Generator, GetTransientPackage()));
+		SetValue(Context, OutGenerator, &Generator);
+	}
+}
+
+void FMakeSphereBoneGeometryGenerator::Register()
+{
+	AddOutputs(Generator);
+
+	Generator = NewObject<UBoneGeometryGenerator_Sphere>(GetOwner());
+}
+
+////////////////////////////////////////////////////////////////////////////////////////////////////////////////
+
+void FMakeCapsuleBoneGeometryGenerator::Evaluate(UE::Dataflow::FContext& Context, const FDataflowOutput* Out) const
+{
+	if(Out->IsA(&Generator))
+	{
+		TObjectPtr<UBoneGeometryGenerator> OutGenerator = CastChecked<UBoneGeometryGenerator>(StaticDuplicateObject(Generator, GetTransientPackage()));
+		SetValue(Context, OutGenerator, &Generator);
+	}
+}
+
+void FMakeCapsuleBoneGeometryGenerator::Register()
+{
+	AddOutputs(Generator);
+
+	Generator = NewObject<UBoneGeometryGenerator_Capsule>(GetOwner());
+}
+
+////////////////////////////////////////////////////////////////////////////////////////////////////////////////
+
+void FMakeConvexBoneGeometryGenerator::Evaluate(UE::Dataflow::FContext& Context, const FDataflowOutput* Out) const
+{
+	if(Out->IsA(&Generator))
+	{
+		TObjectPtr<UBoneGeometryGenerator> OutGenerator = CastChecked<UBoneGeometryGenerator>(StaticDuplicateObject(Generator, GetTransientPackage()));
+		SetValue(Context, OutGenerator, &Generator);
+	}
+}
+
+void FMakeConvexBoneGeometryGenerator::Register()
+{
+	AddOutputs(Generator);
+
+	Generator = NewObject<UBoneGeometryGenerator_Convex>(GetOwner());
+}
+
+////////////////////////////////////////////////////////////////////////////////////////////////////////////////
+
+/**
+ * Data extracted from an LOD for body generation
+ */
+struct FBoneVertData
+{
+	TArray<FVector3f> Positions;
+	TArray<FVector3f> Normals;
+	TArray<float> Weights;
+};
+
+FBoneVertData GetBoneVertexData(FName BoneName, USkeletalMesh* Mesh, int32 Lod, EVertexSelectMode Mode)
+{
+	if(!Mesh)
+	{
+		return {};
+	}
+
+	const FSkeletalMeshRenderData* MeshRenderData = Mesh->GetResourceForRendering();
+	if(!MeshRenderData || !MeshRenderData->LODRenderData.IsValidIndex(Lod))
+	{
+		return {};
+	}
+
+	FBoneVertData Result;
+
+	// Make sure we have initialised the mesh in order to get vertex data
+	Mesh->CalculateInvRefMatrices();
+	const FReferenceSkeleton& RefSkeleton = Mesh->GetRefSkeleton();
+	const FSkeletalMeshLODRenderData& LodData = MeshRenderData->LODRenderData[Lod];
+
+	constexpr float WeightsToFloatMult = (1.0f / 65535.0f);
+	const FSkinWeightVertexBuffer* Weights = LodData.GetSkinWeightVertexBuffer();
+	check(Weights);
+
+	// Mesh buffer for normals, position buffer for positions
+	const FStaticMeshVertexBuffer& MeshVertexBuffer = LodData.StaticVertexBuffers.StaticMeshVertexBuffer;
+	const FPositionVertexBuffer& Positions = LodData.StaticVertexBuffers.PositionVertexBuffer;
+
+	const uint32 NumLodVerts = Weights->GetNumVertices();
+	const uint32 MaxInfluences = LodData.GetVertexBufferMaxBoneInfluences();
+
+	if(MaxInfluences == 0)
+	{
+		return {};
+	}
+
+	const int32 NumSections = LodData.RenderSections.Num();
+
+	for(int32 SectionIndex = 0; SectionIndex < NumSections; ++SectionIndex)
+	{
+		const FSkelMeshRenderSection& Section = LodData.RenderSections[SectionIndex];
+		const uint32 SectionMaxInfluences = Section.MaxBoneInfluences;
+
+		for(uint32 VertIndex = Section.BaseVertexIndex; VertIndex < Section.BaseVertexIndex + Section.NumVertices; ++VertIndex)
+		{
+			if(Mode == EVertexSelectMode::Any)
+			{
+				for(uint32 InfluenceIndex = 0; InfluenceIndex < SectionMaxInfluences; InfluenceIndex++)
+				{
+					const uint32 RenderBoneIndex = Weights->GetBoneIndex(VertIndex, InfluenceIndex);
+					const uint32 MeshBoneIndex = Section.BoneMap[RenderBoneIndex];
+
+					uint16 BoneWeight = Weights->GetBoneWeight(VertIndex, InfluenceIndex);
+
+					if(BoneWeight > 0 && RefSkeleton.GetBoneName(MeshBoneIndex) == BoneName)
+					{
+						const FMatrix44f& RefMatrix = Mesh->GetRefBasesInvMatrix()[MeshBoneIndex];
+
+						Result.Positions.Add(RefMatrix.TransformPosition(Positions.VertexPosition(VertIndex)));
+						Result.Normals.Add(RefMatrix.TransformVector(MeshVertexBuffer.VertexTangentZ(VertIndex)));
+						Result.Weights.Add(BoneWeight * WeightsToFloatMult);
+
+						// Vertex will only have one influence for a given bone, skip the rest
+						break;
+					}
+				}
+			}
+			else
+			{
+				int32 MaxInfluenceIndex = 0;
+				uint16 MaxInfluence = Weights->GetBoneWeight(VertIndex, 0);
+				for(uint32 InfluenceIndex = 1; InfluenceIndex < SectionMaxInfluences; InfluenceIndex++)
+				{
+					uint16 BoneWeight = Weights->GetBoneWeight(VertIndex, InfluenceIndex);
+
+					if(BoneWeight > MaxInfluence)
+					{
+						MaxInfluenceIndex = InfluenceIndex;
+						MaxInfluence = BoneWeight;
+					}
+				}
+
+				const uint32 RenderBoneIndex = Weights->GetBoneIndex(VertIndex, MaxInfluenceIndex);
+				const uint32 MeshBoneIndex = Section.BoneMap[RenderBoneIndex];
+
+				if(RefSkeleton.GetBoneName(MeshBoneIndex) == BoneName)
+				{
+					const FMatrix44f& RefMatrix = Mesh->GetRefBasesInvMatrix()[MeshBoneIndex];
+
+					Result.Positions.Add(RefMatrix.TransformPosition(Positions.VertexPosition(VertIndex)));
+					Result.Normals.Add(RefMatrix.TransformVector(MeshVertexBuffer.VertexTangentZ(VertIndex)));
+					Result.Weights.Add(MaxInfluence * WeightsToFloatMult);
+				}
+			}
+		}
+	}
+
+	return Result;
+}
+
+////////////////////////////////////////////////////////////////////////////////////////////////////////////////
+
+TArray<UE::Chaos::RigidAsset::FSimpleGeometry> UBoneGeometryGenerator_Box::Build(FRigidAssetBoneSelection Bones)
+{
+	return {};
+}
+
+////////////////////////////////////////////////////////////////////////////////////////////////////////////////
+
+TArray<UE::Chaos::RigidAsset::FSimpleGeometry> UBoneGeometryGenerator_Sphere::Build(FRigidAssetBoneSelection Bones)
+{
+	if(!Bones.Mesh || !Bones.Skeleton)
+	{
+		return {};
+	}
+
+	TArray<UE::Chaos::RigidAsset::FSimpleGeometry> Result;
+
+	for(const FRigidAssetBoneInfo& Bone : Bones.SelectedBones)
+	{
+		FBoneVertData BoneVerts = GetBoneVertexData(Bone.Name, Bones.Mesh, BaseSettings.SourceLod, BaseSettings.VertexMode);
+
+		if(BoneVerts.Positions.Num() == 0)
+		{
+			// No influences, make a small sphere and warn
+			// #TODO pass context through for dataflow warnings
+			Result.AddDefaulted();
+			Result.Last().GeometryElement = MakeShared<FKSphereElem>(5.0f);
+			continue;
+		}
+
+		FVector2f AxisIntervals[] =
+		{
+			{std::numeric_limits<float>::max(), std::numeric_limits<float>::lowest()},
+			{std::numeric_limits<float>::max(), std::numeric_limits<float>::lowest()},
+			{std::numeric_limits<float>::max(), std::numeric_limits<float>::lowest()}
+		};
+
+		for(const FVector3f& Position : BoneVerts.Positions)
+		{
+			AxisIntervals[0] = { FMath::Min(AxisIntervals[0][0], Position[0]), FMath::Max(AxisIntervals[0][1], Position[0]) };
+			AxisIntervals[1] = { FMath::Min(AxisIntervals[1][0], Position[1]), FMath::Max(AxisIntervals[1][1], Position[1]) };
+			AxisIntervals[2] = { FMath::Min(AxisIntervals[2][0], Position[2]), FMath::Max(AxisIntervals[2][1], Position[2]) };
+		}
+
+		const FVector3f SphereCentre(
+			(AxisIntervals[0][0] + AxisIntervals[0][1]) * 0.5f,
+			(AxisIntervals[1][0] + AxisIntervals[1][1]) * 0.5f,
+			(AxisIntervals[2][0] + AxisIntervals[2][1]) * 0.5f
+		);
+
+		const float SphereRadius = FMath::Max(AxisIntervals[0][1] - AxisIntervals[0][0], AxisIntervals[1][1] - AxisIntervals[1][0], AxisIntervals[2][1] - AxisIntervals[2][0]);
+
+		Result.AddDefaulted();
+		Result.Last().GeometryElement = MakeShared<FKSphereElem>(SphereRadius);
+	}
+
+	return Result;
+}
+
+////////////////////////////////////////////////////////////////////////////////////////////////////////////////
+
+TArray<UE::Chaos::RigidAsset::FSimpleGeometry> UBoneGeometryGenerator_Capsule::Build(FRigidAssetBoneSelection Bones)
+{
+	return {};
+}
+
+////////////////////////////////////////////////////////////////////////////////////////////////////////////////
+
+TArray<UE::Chaos::RigidAsset::FSimpleGeometry> UBoneGeometryGenerator_Convex::Build(FRigidAssetBoneSelection Bones)
+{
+	return {};
+}
diff --git a/Engine/Plugins/Experimental/ChaosRigidAsset/Source/ChaosRigidAssetNodes/Private/Generators/BoneGeometryGenerators.h b/Engine/Plugins/Experimental/ChaosRigidAsset/Source/ChaosRigidAssetNodes/Private/Generators/BoneGeometryGenerators.h
new file mode 100644
--- /dev/null
+++ b/Engine/Plugins/Experimental/ChaosRigidAsset/Source/ChaosRigidAssetNodes/Private/Generators/BoneGeometryGenerators.h
@@ -0,0 +1,279 @@
+// Copyright Epic Games, Inc. All Rights Reserved.
+
+#include "BoneSelection.h"
+#include "RigidDataflowNode.h"
+#include "ShapeElemNodes.h"
+#include "BoneGeometryGenerators.generated.h"
+
+#pragma once
+
+namespace UE::Dataflow
+{
+	void RegisterBoneGeometryGeneratorNodes();
+}
+
+////////////////////////////////////////////////////////////////////////////////////////////////////////////////
+
+/**
+ * Operation to use for small bones
+ * Skip will ignore a bone, and all of its geometry when attempting to generate collision shapes
+ * Merge will use the corresponding merge thresholds in the generation settings to decide whether to roll bones up
+ * as a single geometry and use all of their render data as a source for the generation process
+ */
+UENUM()
+enum class ESmallBoneOperation
+{
+	Skip,
+	Merge
+};
+
+////////////////////////////////////////////////////////////////////////////////////////////////////////////////
+
+/**
+ * Operation for vertex extraction
+ * Any will extract all verts that have an influence for the selected bone
+ * DominantOnly will only extract verts that have the selcted bone as the dominant weight
+ */
+UENUM()
+enum class EVertexSelectMode
+{
+	Any,
+	DominantOnly
+};
+
+////////////////////////////////////////////////////////////////////////////////////////////////////////////////
+
+/**
+ * Generic configuration for all generator types
+ */
+USTRUCT()
+struct FBaseGenerationSettings
+{
+	GENERATED_BODY()
+
+	// Bones under this limit will be classified as a "Small" bone and apply the requested operation
+	UPROPERTY(EditAnywhere, Category = "Generation")
+	float MinimumBoneSize = 20.0f;
+
+	// Which operation to apply to small bones
+	UPROPERTY(EditAnywhere, Category = "Generation")
+	ESmallBoneOperation SmallBoneOp = ESmallBoneOperation::Skip;
+
+	// For the "Merge" small-bone operation - only merge over this threshold and skip otherwise
+	UPROPERTY(EditAnywhere, Category = "Generation")
+	float SmallBoneMergeThresholdOverride = 0.0f;
+
+	// Vertex extraction operation, see EVertexSelectMode
+	UPROPERTY(EditAnywhere, Category = "Generation")
+	EVertexSelectMode VertexMode = EVertexSelectMode::DominantOnly;
+
+	// The render LOD to use for geometry generation, render data will be extracted from this LOD
+	UPROPERTY(EditAnywhere, Category = "Generation")
+	int32 SourceLod = 0;
+
+	// Extra thickness added to generated geometry after the genration process produces a geometry
+	UPROPERTY(EditAnywhere, Category = "Generation")
+	float Thickness = 0.0f;
+};
+
+/**
+ * Base geometry generator for each render->physics generation type.
+ */
+UCLASS(EditInlineNew, Abstract)
+class UBoneGeometryGenerator : public UObject
+{
+	GENERATED_BODY()
+
+public:
+
+	virtual TArray<UE::Chaos::RigidAsset::FSimpleGeometry> Build(FRigidAssetBoneSelection Bones)
+	{
+		return {};
+	}
+
+protected:
+
+	// Base settings for all generation types
+	UPROPERTY(EditAnywhere, Category = "Generation", meta = (EditInline))
+	FBaseGenerationSettings BaseSettings;
+};
+
+////////////////////////////////////////////////////////////////////////////////////////////////////////////////
+
+/**
+ * Basic box generator
+ * Takes the extracted verts and wraps a tight box around them
+ */
+UCLASS()
+class UBoneGeometryGenerator_Box : public UBoneGeometryGenerator
+{
+	GENERATED_BODY()
+
+public:
+
+	TArray<UE::Chaos::RigidAsset::FSimpleGeometry> Build(FRigidAssetBoneSelection Bones) final override;
+
+private:
+};
+
+////////////////////////////////////////////////////////////////////////////////////////////////////////////////
+
+/**
+ * Basic sphere generator
+ * Takes the extracted verts and wraps a sphere around them
+ */
+UCLASS()
+class UBoneGeometryGenerator_Sphere : public UBoneGeometryGenerator
+{
+	GENERATED_BODY()
+
+public:
+
+	TArray<UE::Chaos::RigidAsset::FSimpleGeometry> Build(FRigidAssetBoneSelection Bones) final override;
+
+private:
+
+	UPROPERTY(EditAnywhere, Category = Settings)
+	bool bUseGeometry = true;
+};
+
+////////////////////////////////////////////////////////////////////////////////////////////////////////////////
+
+/**
+ * Capsule generator
+ * Takes the extracted verts and wraps capsule around them
+ */
+UCLASS()
+class UBoneGeometryGenerator_Capsule : public UBoneGeometryGenerator
+{
+	GENERATED_BODY()
+
+public:
+
+	TArray<UE::Chaos::RigidAsset::FSimpleGeometry> Build(FRigidAssetBoneSelection Bones) final override;
+
+private:
+};
+
+////////////////////////////////////////////////////////////////////////////////////////////////////////////////
+
+/**
+ * Convex generator
+ * Takes the extracted verts then uses them to generate a convex hull decomposition
+ */
+UCLASS()
+class UBoneGeometryGenerator_Convex : public UBoneGeometryGenerator
+{
+	GENERATED_BODY()
+
+public:
+
+	TArray<UE::Chaos::RigidAsset::FSimpleGeometry> Build(FRigidAssetBoneSelection Bones) final override;
+
+private:
+};
+
+////////////////////////////////////////////////////////////////////////////////////////////////////////////////
+
+/**
+ * Make nodes for each generator, each creates a configured generator that can be wired into the generation node
+ */
+USTRUCT()
+struct FMakeBoxBoneGeometryGenerator : public FRigidDataflowNode
+{
+	GENERATED_BODY()
+	DATAFLOW_NODE_DEFINE_INTERNAL(FMakeBoxBoneGeometryGenerator, "Make Box Bone Geometry Builder", "PhysicsAsset", "")
+
+public:
+
+	FMakeBoxBoneGeometryGenerator(const UE::Dataflow::FNodeParameters& InParam, FGuid InGuid = FGuid::NewGuid())
+		: FRigidDataflowNode(InParam, InGuid)
+	{
+		Register();
+	}
+
+private:
+
+	UPROPERTY(EditAnywhere, Instanced, Category = "Generator", meta = (DataflowOutput))
+	TObjectPtr<UBoneGeometryGenerator> Generator;
+
+	void Evaluate(UE::Dataflow::FContext& Context, const FDataflowOutput* Out) const override;
+	void Register();
+};
+
+////////////////////////////////////////////////////////////////////////////////////////////////////////////////
+
+USTRUCT()
+struct FMakeSphereBoneGeometryGenerator : public FRigidDataflowNode
+{
+	GENERATED_BODY()
+	DATAFLOW_NODE_DEFINE_INTERNAL(FMakeSphereBoneGeometryGenerator, "Make Sphere Bone Geometry Builder", "PhysicsAsset", "")
+
+public:
+
+	FMakeSphereBoneGeometryGenerator(const UE::Dataflow::FNodeParameters& InParam, FGuid InGuid = FGuid::NewGuid())
+		: FRigidDataflowNode(InParam, InGuid)
+	{
+		Register();
+	}
+
+private:
+
+	UPROPERTY(EditAnywhere, Instanced, Category = "Generator", meta = (DataflowOutput))
+	TObjectPtr<UBoneGeometryGenerator> Generator;
+
+	void Evaluate(UE::Dataflow::FContext& Context, const FDataflowOutput* Out) const override;
+	void Register();
+};
+
+////////////////////////////////////////////////////////////////////////////////////////////////////////////////
+
+USTRUCT()
+struct FMakeCapsuleBoneGeometryGenerator : public FRigidDataflowNode
+{
+	GENERATED_BODY()
+	DATAFLOW_NODE_DEFINE_INTERNAL(FMakeCapsuleBoneGeometryGenerator, "Make Capsule Bone Geometry Builder", "PhysicsAsset", "")
+
+public:
+
+	FMakeCapsuleBoneGeometryGenerator(const UE::Dataflow::FNodeParameters& InParam, FGuid InGuid = FGuid::NewGuid())
+		: FRigidDataflowNode(InParam, InGuid)
+	{
+		Register();
+	}
+
+private:
+
+	UPROPERTY(EditAnywhere, Instanced, Category = "Generator", meta = (DataflowOutput))
+	TObjectPtr<UBoneGeometryGenerator> Generator;
+
+	void Evaluate(UE::Dataflow::FContext& Context, const FDataflowOutput* Out) const override;
+	void Register();
+};
+
+////////////////////////////////////////////////////////////////////////////////////////////////////////////////
+
+USTRUCT()
+struct FMakeConvexBoneGeometryGenerator : public FRigidDataflowNode
+{
+	GENERATED_BODY()
+	DATAFLOW_NODE_DEFINE_INTERNAL(FMakeConvexBoneGeometryGenerator, "Make Convex Bone Geometry Builder", "PhysicsAsset", "")
+
+public:
+
+	FMakeConvexBoneGeometryGenerator(const UE::Dataflow::FNodeParameters& InParam, FGuid InGuid = FGuid::NewGuid())
+		: FRigidDataflowNode(InParam, InGuid)
+	{
+		Register();
+	}
+
+private:
+
+	UPROPERTY(EditAnywhere, Instanced, Category = "Generator", meta = (DataflowOutput))
+	TObjectPtr<UBoneGeometryGenerator> Generator;
+
+	void Evaluate(UE::Dataflow::FContext& Context, const FDataflowOutput* Out) const override;
+	void Register();
+};
+
+////////////////////////////////////////////////////////////////////////////////////////////////////////////////
diff --git a/Engine/Plugins/Experimental/ChaosRigidAsset/Source/ChaosRigidAssetNodes/Private/Generators/ConstraintGenerators.cpp b/Engine/Plugins/Experimental/ChaosRigidAsset/Source/ChaosRigidAssetNodes/Private/Generators/ConstraintGenerators.cpp
new file mode 100644
--- /dev/null
+++ b/Engine/Plugins/Experimental/ChaosRigidAsset/Source/ChaosRigidAssetNodes/Private/Generators/ConstraintGenerators.cpp
@@ -0,0 +1,100 @@
+// Copyright Epic Games, Inc. All Rights Reserved.
+
+#include "Generators/ConstraintGenerators.h"
+
+#include "Dataflow/DataflowNode.h"
+#include "Dataflow/DataflowNodeFactory.h"
+#include "Engine/SkeletalMesh.h"
+#include "ReferenceSkeleton.h"
+
+namespace UE::Dataflow
+{
+	void RegisterConstraintGeneratorNodes()
+	{
+		// Constraint Generators
+		DATAFLOW_NODE_REGISTER_CREATION_FACTORY(FMakeSwingTwistConstraintGenerator);
+	}
+}
+
+////////////////////////////////////////////////////////////////////////////////////////////////////////////////
+
+TArray<TObjectPtr<UPhysicsConstraintTemplate>> UConstraintGenerator_SwingTwist::Build(TObjectPtr<UPhysicsConstraintTemplate> ConstraintTemplate, FRigidAssetBoneSelection Bones) const
+{
+	TArray<TObjectPtr<UPhysicsConstraintTemplate>> Result;
+
+	// Look at the requested bones and create a constraint to the first parent we find for each
+	const FReferenceSkeleton& RefSkel = Bones.Mesh->GetRefSkeleton();
+	for(const FRigidAssetBoneInfo& Bone : Bones.SelectedBones)
+	{
+		int32 CurrParent = RefSkel.GetParentIndex(Bone.Index);
+
+		while(CurrParent != INDEX_NONE)
+		{
+			const FRigidAssetBoneInfo* ParentBone = Bones.SelectedBones.FindByPredicate([&CurrParent](const FRigidAssetBoneInfo& Item)
+				{
+					return Item.Index == CurrParent;
+				});
+
+			if(!ParentBone)
+			{
+				CurrParent = RefSkel.GetParentIndex(CurrParent);
+				continue;
+			}
+
+			// Now we have a parent, and it's in the selection - add a constraint
+			TObjectPtr<UPhysicsConstraintTemplate> NewConstraint = NewObject<UPhysicsConstraintTemplate>();
+			FConstraintInstance& Defaults = NewConstraint->DefaultInstance;
+
+			if(ConstraintTemplate)
+			{
+				Defaults.CopyConstraintParamsFrom(&ConstraintTemplate->DefaultInstance);
+			}
+
+			Defaults.SetAngularSwing1Limit(EAngularConstraintMotion::ACM_Limited, Swing1Limit);
+			Defaults.SetAngularSwing2Limit(EAngularConstraintMotion::ACM_Limited, Swing2Limit);
+			Defaults.SetAngularTwistLimit(EAngularConstraintMotion::ACM_Limited, TwistLimit);
+
+			Defaults.ConstraintBone1 = Bone.Name;
+			Defaults.ConstraintBone2 = ParentBone->Name;
+
+			FTransform ParentTm = CalculateRelativeBoneTransform(Bone.Name, ParentBone->Name, RefSkel);
+			FTransform ChildTm = FTransform::Identity;
+
+			Defaults.SetRefPosition(EConstraintFrame::Frame1, ChildTm.GetTranslation());
+			Defaults.SetRefOrientation(EConstraintFrame::Frame1, ChildTm.GetUnitAxis(EAxis::X), ChildTm.GetUnitAxis(EAxis::Y));
+			Defaults.SetRefPosition(EConstraintFrame::Frame2, ParentTm.GetTranslation());
+			Defaults.SetRefOrientation(EConstraintFrame::Frame2, ParentTm.GetUnitAxis(EAxis::X), ParentTm.GetUnitAxis(EAxis::Y));
+
+#if WITH_EDITOR
+			// In editor context, set the default instance for the constraint. This is not available outside of editor
+			NewConstraint->SetDefaultProfile(Defaults);
+#endif
+
+			Result.Add(NewConstraint);
+
+			break;
+		}
+	}
+
+	return Result;
+}
+
+////////////////////////////////////////////////////////////////////////////////////////////////////////////////
+
+void FMakeSwingTwistConstraintGenerator::Evaluate(UE::Dataflow::FContext& Context, const FDataflowOutput* Out) const
+{
+	if(Out->IsA(&Generator))
+	{
+		TObjectPtr<UConstraintGenerator> OutGenerator = CastChecked<UConstraintGenerator>(StaticDuplicateObject(Generator, GetTransientPackage()));
+		SetValue(Context, OutGenerator, &Generator);
+	}
+}
+
+void FMakeSwingTwistConstraintGenerator::Register()
+{
+	AddOutputs(Generator);
+
+	Generator = NewObject<UConstraintGenerator_SwingTwist>(GetOwner());
+}
+
+////////////////////////////////////////////////////////////////////////////////////////////////////////////////
diff --git a/Engine/Plugins/Experimental/ChaosRigidAsset/Source/ChaosRigidAssetNodes/Private/Generators/ConstraintGenerators.h b/Engine/Plugins/Experimental/ChaosRigidAsset/Source/ChaosRigidAssetNodes/Private/Generators/ConstraintGenerators.h
new file mode 100644
--- /dev/null
+++ b/Engine/Plugins/Experimental/ChaosRigidAsset/Source/ChaosRigidAssetNodes/Private/Generators/ConstraintGenerators.h
@@ -0,0 +1,105 @@
+// Copyright Epic Games, Inc. All Rights Reserved.
+
+#pragma once
+
+#include "PhysicsEngine/PhysicsConstraintTemplate.h"
+#include "BoneSelection.h"
+#include "RigidDataflowNode.h"
+
+#include "ConstraintGenerators.generated.h"
+
+namespace UE::Dataflow
+{
+	void RegisterConstraintGeneratorNodes();
+}
+
+////////////////////////////////////////////////////////////////////////////////////////////////////////////////
+
+/**
+ * Shared settings for all generator types
+ */
+USTRUCT()
+struct FBaseConstraintGenerationSettings
+{
+	GENERATED_BODY();
+};
+
+////////////////////////////////////////////////////////////////////////////////////////////////////////////////
+
+/**
+ * Base class for constraint generators
+ */
+UCLASS(EditInlineNew, Abstract)
+class UConstraintGenerator : public UObject
+{
+	GENERATED_BODY()
+
+public:
+
+	virtual TArray<TObjectPtr<UPhysicsConstraintTemplate>> Build(TObjectPtr<UPhysicsConstraintTemplate> ConstraintTemplate, FRigidAssetBoneSelection Bones) const
+	{ 
+		return {};
+	}
+
+private:
+
+	UPROPERTY(EditAnywhere, Category = "Generation", meta=(EditInline))
+	FBaseConstraintGenerationSettings BaseSettings;
+};
+
+////////////////////////////////////////////////////////////////////////////////////////////////////////////////
+
+/**
+ * Swing/Twist constraint generator
+ * Given limits for each of the angular axes, this generator will link all the bones in the input set together with the specified limits
+ */
+UCLASS()
+class UConstraintGenerator_SwingTwist : public UConstraintGenerator
+{
+	GENERATED_BODY()
+
+public:
+
+	TArray<TObjectPtr<UPhysicsConstraintTemplate>> Build(TObjectPtr<UPhysicsConstraintTemplate> ConstraintTemplate, FRigidAssetBoneSelection Bones) const override;
+
+private:
+
+	UPROPERTY(EditAnywhere, Category=Constraint)
+	float Swing1Limit;
+
+	UPROPERTY(EditAnywhere, Category=Constraint)
+	float Swing2Limit;
+
+	UPROPERTY(EditAnywhere, Category=Constraint)
+	float TwistLimit;
+};
+
+////////////////////////////////////////////////////////////////////////////////////////////////////////////////
+
+/**
+ * Generator make nodes
+ */
+USTRUCT()
+struct FMakeSwingTwistConstraintGenerator : public FRigidDataflowNode
+{
+	GENERATED_BODY()
+	DATAFLOW_NODE_DEFINE_INTERNAL(FMakeSwingTwistConstraintGenerator, "Make Swing/Twist Constraint Generator", "PhysicsAsset", "")
+
+public:
+
+	FMakeSwingTwistConstraintGenerator(const UE::Dataflow::FNodeParameters& InParam, FGuid InGuid = FGuid::NewGuid())
+		: FRigidDataflowNode(InParam, InGuid)
+	{
+		Register();
+	}
+
+private:
+
+	UPROPERTY(EditAnywhere, Instanced, Category = "Generator", meta = (DataflowOutput))
+	TObjectPtr<UConstraintGenerator> Generator;
+
+	void Evaluate(UE::Dataflow::FContext& Context, const FDataflowOutput* Out) const override;
+	void Register();
+};
+
+////////////////////////////////////////////////////////////////////////////////////////////////////////////////
diff --git a/Engine/Plugins/Experimental/ChaosRigidAsset/Source/ChaosRigidAssetNodes/Private/PhysicsAssetBuilder.cpp b/Engine/Plugins/Experimental/ChaosRigidAsset/Source/ChaosRigidAssetNodes/Private/PhysicsAssetBuilder.cpp
new file mode 100644
--- /dev/null
+++ b/Engine/Plugins/Experimental/ChaosRigidAsset/Source/ChaosRigidAssetNodes/Private/PhysicsAssetBuilder.cpp
@@ -0,0 +1,172 @@
+// Copyright Epic Games, Inc. All Rights Reserved.
+
+#include "PhysicsAssetBuilder.h"
+#include "PhysicsEngine/PhysicsAsset.h"
+#include "PhysicsEngine/SkeletalBodySetup.h"
+#include "PhysicsEngine/PhysicsConstraintTemplate.h"
+#include "Engine/SkeletalMesh.h"
+#include "UObject/Package.h"
+
+#if WITH_EDITOR
+#include "IAssetTools.h"
+#include "AssetToolsModule.h"
+#endif
+
+namespace UE::Chaos::RigidAsset
+{
+	FPhysicsAssetBuilder FPhysicsAssetBuilder::Make(TObjectPtr<USkeleton> InTargetSkeleton)
+	{
+		FPhysicsAssetBuilder Builder;
+		Builder.TargetSkeleton = InTargetSkeleton;
+
+		return Builder;
+	}
+
+	FPhysicsAssetBuilder FPhysicsAssetBuilder::Make(TObjectPtr<USkeleton> InTargetSkeleton, const TArray<TObjectPtr<USkeletalBodySetup>>& InBodies, const TArray<TObjectPtr<UPhysicsConstraintTemplate>>& InConstraints)
+	{
+		FPhysicsAssetBuilder Builder;
+
+		Builder.TargetSkeleton = InTargetSkeleton;
+		Builder.Bodies = InBodies;
+		Builder.Constraints = InConstraints;
+
+		return Builder;
+	}
+
+	FPhysicsAssetBuilder& FPhysicsAssetBuilder::Body(TObjectPtr<USkeletalBodySetup> NewBody)
+	{
+		Bodies.Add(NewBody);
+
+		return *this;
+	}
+
+	FPhysicsAssetBuilder& FPhysicsAssetBuilder::Joint(TObjectPtr<UPhysicsConstraintTemplate> NewJoint)
+	{
+		Constraints.Add(NewJoint);
+
+		return *this;
+	}
+
+	FPhysicsAssetBuilder& FPhysicsAssetBuilder::JoinLast(TObjectPtr<UPhysicsConstraintTemplate> NewJoint)
+	{
+		// We need two bodies to join
+		const int32 NumBodies = Bodies.Num();
+		if(NumBodies > 1)
+		{
+			TObjectPtr<USkeletalBodySetup> BodiesToJoin[] = { Bodies[NumBodies - 2], Bodies[NumBodies - 1] };
+
+			NewJoint->DefaultInstance.ConstraintBone1 = BodiesToJoin[0]->BoneName;
+			NewJoint->DefaultInstance.ConstraintBone2 = BodiesToJoin[1]->BoneName;
+
+			Constraints.Add(NewJoint);
+		}
+		else
+		{
+			UE_LOG(LogTemp, Error, TEXT("JoinLast called with less than two available bodies in the builder - cannot join these bodies"));
+		}
+
+		return *this;
+	}
+
+	FPhysicsAssetBuilder& FPhysicsAssetBuilder::Path(FString InDesiredPath)
+	{
+		DesiredPath = InDesiredPath;
+
+		return *this;
+	}
+
+	FPhysicsAssetBuilder& FPhysicsAssetBuilder::Path(TObjectPtr<UObject> InObjectReferenceForPath)
+	{
+		DesiredPath = InObjectReferenceForPath->GetPathName();
+
+		return *this;
+	}
+
+	FPhysicsAssetBuilder& FPhysicsAssetBuilder::SetTargetAsset(TObjectPtr<UPhysicsAsset> InAsset)
+	{
+		TargetAsset = InAsset;
+
+		return *this;
+	}
+
+	TObjectPtr<UPhysicsAsset> FPhysicsAssetBuilder::Build()
+	{
+		if(!TargetAsset)
+		{
+			CreateNewAsset();
+		}
+		check(TargetAsset);
+
+		// Need to replace the objects already in the asset, first move them to the transient package to get them out of the physics asset package
+		for(TObjectPtr<USkeletalBodySetup> BodyToRename : TargetAsset->SkeletalBodySetups)
+		{
+			verify(BodyToRename->Rename(nullptr, GetTransientPackage(), REN_DontCreateRedirectors));
+		}
+
+		for(TObjectPtr<UPhysicsConstraintTemplate> ConstraintToRename : TargetAsset->ConstraintSetup)
+		{
+			verify(ConstraintToRename->Rename(nullptr, GetTransientPackage(), REN_DontCreateRedirectors));
+		}
+
+		// Swap to new setups and rename them into the asset package
+		TargetAsset->SkeletalBodySetups = Bodies;
+		TargetAsset->ConstraintSetup = Constraints;
+
+		for(TObjectPtr<USkeletalBodySetup> BodyToRename : TargetAsset->SkeletalBodySetups)
+		{
+			verify(BodyToRename->Rename(nullptr, TargetAsset, REN_DontCreateRedirectors));
+		}
+
+		for(TObjectPtr<UPhysicsConstraintTemplate> ConstraintToRename : TargetAsset->ConstraintSetup)
+		{
+			verify(ConstraintToRename->Rename(nullptr, TargetAsset, REN_DontCreateRedirectors));
+		}
+
+		TargetAsset->UpdateBodySetupIndexMap();
+		TargetAsset->UpdateBoundsBodiesArray();
+
+		// Disable collisions between directly constrained joints
+		for(TObjectPtr<UPhysicsConstraintTemplate> Constraint : TargetAsset->ConstraintSetup)
+		{
+			const int32 Indices[] = {
+				TargetAsset->FindBodyIndex(Constraint->DefaultInstance.ConstraintBone1),
+				TargetAsset->FindBodyIndex(Constraint->DefaultInstance.ConstraintBone2)
+			};
+
+			if(Indices[0] == INDEX_NONE || Indices[1] == INDEX_NONE)
+			{
+				continue;
+			}
+
+			TargetAsset->DisableCollision(Indices[0], Indices[1]);
+		}
+
+		RefreshSkelMeshOnPhysicsAssetChange(TargetAsset->GetPreviewMesh());
+
+		return TargetAsset;
+	}
+
+	FPhysicsAssetBuilder::FPhysicsAssetBuilder()
+	{
+	}
+
+	void FPhysicsAssetBuilder::CreateNewAsset()
+	{
+#if WITH_EDITOR
+		IAssetTools& AssetTools = FModuleManager::LoadModuleChecked<FAssetToolsModule>("AssetTools").Get();
+		FString UniquePackageName;
+		FString UniqueAssetName;
+
+		if(DesiredPath.IsEmpty())
+		{
+			DesiredPath = "/";
+		}
+
+		AssetTools.CreateUniqueAssetName(DesiredPath, TEXT(""), UniquePackageName, UniqueAssetName);
+
+		TargetAsset = CastChecked<UPhysicsAsset>(AssetTools.CreateAsset(*UniqueAssetName, *UniquePackageName, UPhysicsAsset::StaticClass(), nullptr));
+#else
+		ensureMsgf(false, TEXT("Attempted to create a new asset package outside of the editor. At runtime use SetTargetAsset in order to use FPhysicsAssetBuilder"));
+#endif
+	}
+}
diff --git a/Engine/Plugins/Experimental/ChaosRigidAsset/Source/ChaosRigidAssetNodes/Private/PhysicsAssetBuilder.h b/Engine/Plugins/Experimental/ChaosRigidAsset/Source/ChaosRigidAssetNodes/Private/PhysicsAssetBuilder.h
new file mode 100644
--- /dev/null
+++ b/Engine/Plugins/Experimental/ChaosRigidAsset/Source/ChaosRigidAssetNodes/Private/PhysicsAssetBuilder.h
@@ -0,0 +1,60 @@
+// Copyright Epic Games, Inc. All Rights Reserved.
+
+#pragma once
+
+#include "UObject/ObjectPtr.h"
+#include "Containers/Array.h"
+
+class UObject;
+class USkeleton;
+class USkeletalBodySetup;
+class UPhysicsConstraintTemplate;
+class UPhysicsAsset;
+
+namespace UE::Chaos::RigidAsset
+{
+	/**
+	 * Builder for physics assets
+	 * Takes source data for physics assets (bodies, constraints, mesh, skeleton) and creates a UPhysicsAsset (or writes to a target asset)
+	 */
+	class FPhysicsAssetBuilder
+	{
+	public:
+
+		// Overloads for creating a builder from required data.
+		static FPhysicsAssetBuilder Make(TObjectPtr<USkeleton> InTargetSkeleton);
+		static FPhysicsAssetBuilder Make(TObjectPtr<USkeleton> InTargetSkeleton, const TArray<TObjectPtr<USkeletalBodySetup>>& InBodies, const TArray<TObjectPtr<UPhysicsConstraintTemplate>>& InConstraints);
+
+		// Append a new body to the asset
+		FPhysicsAssetBuilder& Body(TObjectPtr<USkeletalBodySetup> NewBody);
+
+		// Append a new constraint to the asset
+		FPhysicsAssetBuilder& Joint(TObjectPtr<UPhysicsConstraintTemplate> NewJoint);
+
+		// Apply the provided constraint template to the last pair of bodies in the builder
+		FPhysicsAssetBuilder& JoinLast(TObjectPtr<UPhysicsConstraintTemplate> NewJoint);
+
+		// If not targetting a specific asset, set the path to where the resulting asset will be stored
+		FPhysicsAssetBuilder& Path(FString InDesiredPath);
+		FPhysicsAssetBuilder& Path(TObjectPtr<UObject> InObjectReferenceForPath);
+
+		// Set a target asset instead of a path - this asset will be overwritten with new state when Build() is called
+		FPhysicsAssetBuilder& SetTargetAsset(TObjectPtr<UPhysicsAsset> InAsset);
+
+		// Finalize the physics asset and apply the builder to it
+		TObjectPtr<UPhysicsAsset> Build();
+
+	private:
+
+		FPhysicsAssetBuilder();
+
+		void CreateNewAsset();
+
+		FString DesiredPath;
+
+		TObjectPtr<USkeleton> TargetSkeleton;
+		TArray<TObjectPtr<USkeletalBodySetup>> Bodies;
+		TArray<TObjectPtr<UPhysicsConstraintTemplate>> Constraints;
+		TObjectPtr<UPhysicsAsset> TargetAsset;
+	};
+}
\ No newline at end of file
diff --git a/Engine/Plugins/Experimental/ChaosRigidAsset/Source/ChaosRigidAssetNodes/Private/PhysicsAssetDataflowNodes.cpp b/Engine/Plugins/Experimental/ChaosRigidAsset/Source/ChaosRigidAssetNodes/Private/PhysicsAssetDataflowNodes.cpp
new file mode 100644
--- /dev/null
+++ b/Engine/Plugins/Experimental/ChaosRigidAsset/Source/ChaosRigidAssetNodes/Private/PhysicsAssetDataflowNodes.cpp
@@ -0,0 +1,831 @@
+// Copyright Epic Games, Inc. All Rights Reserved.
+
+#include "PhysicsAssetDataflowNodes.h"
+#include "Dataflow/DataflowNodeFactory.h"
+#include "PhysicsEngine/PhysicsAsset.h"
+#include "PhysicsEngine/SkeletalBodySetup.h"
+#include "Misc/WildcardString.h"
+#include "DataflowAttachment.h"
+#include "Dataflow/DataflowDebugDrawInterface.h"
+#include "Dataflow/DataflowDebugDrawObject.h"
+#include "Animation/Skeleton.h"
+#include "SkeletalDebugRendering.h"
+#include "PhysicsAssetBuilder.h"
+#include "Engine/SkeletalMesh.h"
+#include "ReferenceSkeleton.h"
+
+FDataflowPhysicsAssetTerminalNode::FDataflowPhysicsAssetTerminalNode(const UE::Dataflow::FNodeParameters& InParam, FGuid InGuid /*= FGuid::NewGuid()*/)
+	: FDataflowTerminalNode(InParam, InGuid)
+{
+	RegisterInputConnection(&State);
+	RegisterInputConnection(&PhysicsAsset);
+}
+
+void FDataflowPhysicsAssetTerminalNode::SetAssetValue(TObjectPtr<UObject> Asset, UE::Dataflow::FContext& Context) const
+{
+	UPhysicsAsset* Target = Asset ? Cast<UPhysicsAsset>(Asset->GetOuter()) : nullptr;
+
+	if(!Target)
+	{
+		Target = GetValue(Context, &PhysicsAsset);
+	}
+
+	if(Target)
+	{
+		const FPhysicsAssetDataflowState InState = GetValue(Context, &State);
+		InState.DebugLog();
+
+		UE::Chaos::RigidAsset::FPhysicsAssetBuilder::Make(InState.TargetSkeleton, InState.Bodies.ToArray(), InState.Constraints.ToArray())
+			.SetTargetAsset(Target)
+			.Build();
+	}
+}
+
+void FDataflowPhysicsAssetTerminalNode::Evaluate(UE::Dataflow::FContext& Context) const
+{
+
+}
+
+void UE::Dataflow::RegisterPhysicsAssetTerminalNode()
+{
+	DATAFLOW_NODE_REGISTER_CREATION_FACTORY(FDataflowPhysicsAssetTerminalNode);
+}
+
+void UE::Dataflow::RegisterPhysicsAssetNodes()
+{
+	// Asset state
+	DATAFLOW_NODE_REGISTER_CREATION_FACTORY(FDataflowPhysicsAssetMakeState);
+	DATAFLOW_NODE_REGISTER_CREATION_FACTORY(FDataflowPhysicsAssetAddBody);
+
+	// Body Setup
+	DATAFLOW_NODE_REGISTER_CREATION_FACTORY(FDataflowMakeBody);
+	DATAFLOW_NODE_REGISTER_CREATION_FACTORY(FDataflowSetBodyGeometry);
+
+	// Joints
+	DATAFLOW_NODE_REGISTER_CREATION_FACTORY(FDataflowMakeJoint);
+	DATAFLOW_NODE_REGISTER_CREATION_FACTORY(FDataflowAutoConstrainBodies);
+
+	// Body Geometry
+	DATAFLOW_NODE_REGISTER_CREATION_FACTORY(FDataflowAggGeomAddShape);
+	DATAFLOW_NODE_REGISTER_CREATION_FACTORY(FDataflowCreateGeometryForBones);
+
+	// Bone Selections
+	DATAFLOW_NODE_REGISTER_CREATION_FACTORY(FDataflowNewBoneSelection);
+	DATAFLOW_NODE_REGISTER_CREATION_FACTORY(FDataflowAppendBoneSelection);
+	DATAFLOW_NODE_REGISTER_CREATION_FACTORY(FDataflowSelectBonesByName);
+	DATAFLOW_NODE_REGISTER_CREATION_FACTORY(FDataflowSelectConnectedBones);
+
+	RegisterBoneGeometryGeneratorNodes();
+	RegisterConstraintGeneratorNodes();
+}
+
+FDataflowPhysicsAssetAddBody::FDataflowPhysicsAssetAddBody(const UE::Dataflow::FNodeParameters& InParam, FGuid InGuid /*= FGuid::NewGuid()*/)
+	: FRigidDataflowNode(InParam, InGuid)
+{
+	AddInputs(Body, State);
+	AddPassthroughs(State);
+}
+
+void FDataflowPhysicsAssetAddBody::Evaluate(UE::Dataflow::FContext& Context, const FDataflowOutput* Out) const
+{
+	if(Out->IsA(&State))
+	{
+		const USkeletalBodySetup* InBody = GetValue(Context, &Body);
+
+		FPhysicsAssetDataflowState InState = GetValue(Context, &State);
+
+		if(InBody)
+		{
+			InState.Bodies.Emplace(Cast<USkeletalBodySetup>(StaticDuplicateObject(InBody, GetTransientPackage())));
+		}
+
+		SetValue(Context, InState, &State);
+	}
+}
+
+struct FDataflowDebugDrawBodySetupObject : public FDataflowDebugDrawBaseObject
+{
+	FDataflowDebugDrawBodySetupObject(IDataflowDebugDrawInterface::FDataflowElementsType& InDataflowElements, TObjectPtr<USkeletalBodySetup> InBody)
+		: FDataflowDebugDrawBaseObject(InDataflowElements)
+		, Body(InBody)
+	{
+	}
+
+	virtual void PopulateDataflowElements() override
+	{
+	}
+
+	virtual void DrawDataflowElements(FPrimitiveDrawInterface* PDI) override
+	{
+		if(!Body)
+		{
+			return;
+		}
+
+		const FKAggregateGeom& Geom = Body->AggGeom;
+		const int32 NumGeoms = Geom.GetElementCount();
+
+		for(int32 Index = 0; Index < NumGeoms; ++Index)
+		{
+			const FKShapeElem* Elem = Geom.GetElement(Index);
+			Elem->DrawElemWire(PDI, Elem->GetTransform(), 1.0f, FColor::White);
+		}
+	}
+
+	virtual FBox ComputeBoundingBox() const override
+	{
+		if(!Body)
+		{
+			return FBox(EForceInit::ForceInitToZero);
+		}
+
+		return Body->AggGeom.CalcAABB(FTransform::Identity);
+	}
+
+private:
+	TObjectPtr<USkeletalBodySetup> Body;
+};
+
+struct FDataflowDebugDrawBoneSelectionObject : public FDataflowDebugDrawBaseObject
+{
+	FDataflowDebugDrawBoneSelectionObject(IDataflowDebugDrawInterface::FDataflowElementsType& InDataflowElements, FRigidAssetBoneSelection InSelection)
+		: FDataflowDebugDrawBaseObject(InDataflowElements)
+		, Selection(InSelection)
+	{
+	}
+
+	virtual void PopulateDataflowElements() override
+	{
+	}
+
+	virtual void DrawDataflowElements(FPrimitiveDrawInterface* PDI) override
+	{
+		if(!Selection.Mesh)
+		{
+			return;
+		}
+
+		FReferenceSkeleton& RefSkel = Selection.Mesh->GetRefSkeleton();
+
+		const int32 NumBones = RefSkel.GetNum();
+		TArray<FTransform> BoneTransforms;
+		TArray<FBoneIndexType> RequiredBones;
+		TArray<FLinearColor> BoneColors;
+		TArray<int32> SelectedBones;
+		TArray<TRefCountPtr<HHitProxy>> HitProxies;
+
+		RefSkel.GetBoneAbsoluteTransforms(BoneTransforms);
+
+		for(int32 BoneIndex = 0; BoneIndex < NumBones; ++BoneIndex)
+		{
+			RequiredBones.Add(BoneIndex);
+		}
+
+		for(const FRigidAssetBoneInfo& Bone : Selection.SelectedBones)
+		{
+			SelectedBones.Add(Bone.Index);
+		}
+
+		FSkelDebugDrawConfig DrawConfig;
+		DrawConfig.bUseMultiColorAsDefaultColor = false;
+		DrawConfig.BoneDrawMode = EBoneDrawMode::All;
+		DrawConfig.BoneDrawSize = 1.f;
+		DrawConfig.bAddHitProxy = false;
+		DrawConfig.bForceDraw = true;
+		DrawConfig.DefaultBoneColor = FLinearColor::Gray;
+		DrawConfig.AffectedBoneColor = FLinearColor::Gray;
+		DrawConfig.SelectedBoneColor = FLinearColor{ 1.0f, 0.75f, 0.1f };
+		DrawConfig.ParentOfSelectedBoneColor = FLinearColor::Gray;
+
+		SkeletalDebugRendering::DrawBones(
+			PDI, FVector::Zero(),
+			RequiredBones,
+			RefSkel,
+			BoneTransforms,
+			SelectedBones,
+			BoneColors,
+			HitProxies,
+			DrawConfig
+		);
+	}
+
+	virtual FBox ComputeBoundingBox() const override
+	{
+		if(!Selection.Mesh)
+		{
+			return FBox(EForceInit::ForceInitToZero);
+		}
+
+		FReferenceSkeleton& RefSkel = Selection.Mesh->GetRefSkeleton();
+		TArray<FTransform> BoneTransforms;
+
+		RefSkel.GetBoneAbsoluteTransforms(BoneTransforms);
+
+		FBox Result;
+
+		for(const FTransform& Transform : BoneTransforms)
+		{
+			Result += Transform.GetTranslation();
+		}
+
+		return Result;
+	}
+
+private:
+	FRigidAssetBoneSelection Selection;
+};
+
+struct FDataflowDebugDrawAggGeomObject : public FDataflowDebugDrawBaseObject
+{
+	FDataflowDebugDrawAggGeomObject(IDataflowDebugDrawInterface::FDataflowElementsType& InDataflowElements, const FKAggregateGeom* InGeom)
+		: FDataflowDebugDrawBaseObject(InDataflowElements)
+		, Geom(InGeom)
+	{
+	}
+
+	void PopulateDataflowElements() override
+	{
+	}
+
+	void DrawDataflowElements(FPrimitiveDrawInterface* PDI) override
+	{
+		if(!Geom)
+		{
+			return;
+		}
+
+		const int32 NumElems = Geom->GetElementCount();
+
+		int32 ElementOffset = DataflowElements.Num();
+		int32 ElementNum = NumElems;
+
+		for(int32 Index = 0; Index < NumElems; ++Index)
+		{
+			const FKShapeElem* Elem = Geom->GetElement(Index);
+			Elem->DrawElemWire(PDI, Elem->GetTransform(), 1.0f, FColor::Cyan);
+		}
+	}
+
+	FBox ComputeBoundingBox() const override
+	{
+		if(!Geom)
+		{
+			return {};
+		}
+
+		return Geom->CalcAABB(FTransform::Identity);
+	}
+
+private:
+	const FKAggregateGeom* Geom = nullptr;
+};
+
+#if WITH_EDITOR
+bool FDataflowMakeBody::CanDebugDraw() const
+{
+	return true;
+}
+
+bool FDataflowMakeBody::CanDebugDrawViewMode(const FName& ViewModeName) const
+{
+	return true;
+}
+
+void FDataflowMakeBody::DebugDraw(UE::Dataflow::FContext& Context, IDataflowDebugDrawInterface& DataflowRenderingInterface, const FDebugDrawParameters& DebugDrawParameters) const
+{
+	if(Body)
+	{
+		TRefCountPtr<IDataflowDebugDrawObject> BodyDrawObject(MakeDebugDrawObject<FDataflowDebugDrawBodySetupObject>(DataflowRenderingInterface.ModifyDataflowElements(), Body));
+		DataflowRenderingInterface.DrawObject(BodyDrawObject);
+	}
+}
+#endif
+
+void FDataflowMakeBody::Evaluate(UE::Dataflow::FContext& Context, const FDataflowOutput* Out) const
+{
+	if(Out->IsA(&Body))
+	{
+		TObjectPtr<USkeletalBodySetup> NewSetup = CastChecked<USkeletalBodySetup>(StaticDuplicateObject(Body, GetTransientPackage()));
+
+		NewSetup->BoneName = GetValue(Context, &BoneName);
+
+		SetValue(Context, NewSetup, &Body);
+	}
+}
+
+void FDataflowMakeBody::Register(const UE::Dataflow::FNodeParameters& InParam)
+{
+	// Construct the template object for this node
+	Body = NewObject<USkeletalBodySetup>(GetOwner());
+
+	AddInputs(BoneName);
+	AddOutputs(Body);
+}
+
+void FDataflowMakeJoint::Evaluate(UE::Dataflow::FContext& Context, const FDataflowOutput* Out) const
+{
+	if(Out->IsA(&Joint))
+	{
+		TObjectPtr<UPhysicsConstraintTemplate> NewJoint = CastChecked<UPhysicsConstraintTemplate>(StaticDuplicateObject(Joint, GetTransientPackage()));
+
+		SetValue(Context, NewJoint, &Joint);
+	}
+}
+
+void FDataflowMakeJoint::Register(const UE::Dataflow::FNodeParameters& InParam)
+{
+	// Construct the template object for this node
+	Joint = NewObject<UPhysicsConstraintTemplate>(GetOwner());
+
+	AddOutputs(Joint);
+}
+
+void FDataflowSelectBonesByName::Evaluate(UE::Dataflow::FContext& Context, const FDataflowOutput* Out) const
+{
+	if(Out->IsA(&Selection))
+	{
+		FPhysicsAssetDataflowState InState = GetValue(Context, &State);
+
+		TObjectPtr<USkeleton> Skeleton = InState.TargetSkeleton;
+		TObjectPtr<USkeletalMesh> Mesh = InState.TargetMesh;
+
+		if(!Skeleton)
+		{
+			Context.Error(TEXT("Attempted to make a selection without a Skeleton"), this);
+			SetValue(Context, FRigidAssetBoneSelection{}, &Selection);
+			return;
+		}
+
+		if(!Mesh)
+		{
+			Context.Error(TEXT("Attempted to make a selection without a SkeletalMesh"), this);
+			SetValue(Context, FRigidAssetBoneSelection{}, &Selection);
+			return;
+		}
+
+		const FWildcardString Pattern = SearchString;
+		const FReferenceSkeleton RefSkel = Skeleton->GetReferenceSkeleton();
+
+		FRigidAssetBoneSelection NewSelection;
+		NewSelection.Skeleton = Skeleton;
+		NewSelection.Mesh = Mesh;
+
+		for(int32 BoneIndex = 0; BoneIndex < RefSkel.GetNum(); ++BoneIndex)
+		{
+			FName BoneName = RefSkel.GetBoneName(BoneIndex);
+
+			if(Pattern.IsMatch(BoneName.ToString()))
+			{
+				NewSelection.SelectedBones.Emplace(BoneName, BoneIndex, RefSkel.GetDepthBetweenBones(BoneIndex, 0));
+			}
+		}
+
+		NewSelection.SortBones();
+		SetValue(Context, MoveTemp(NewSelection), &Selection);
+	}
+}
+
+void FDataflowSelectBonesByName::Register(const UE::Dataflow::FNodeParameters& InParam)
+{
+	AddInputs(State);
+	AddOutputs(Selection);
+}
+
+void FDataflowAppendBoneSelection::Evaluate(UE::Dataflow::FContext& Context, const FDataflowOutput* Out) const
+{
+	if(Out->IsA(&Selection))
+	{
+		const FRigidAssetBoneSelection& InA = GetValue(Context, &A);
+		const FRigidAssetBoneSelection& InB = GetValue(Context, &B);
+
+		if((InA.Skeleton != InB.Skeleton) || (InA.Mesh != InB.Mesh))
+		{
+			Context.Error(TEXT("Attmepted to append selections from multiple skeletons or meshes"), this);
+			SetValue(Context, FRigidAssetBoneSelection{}, &Selection);
+			return;
+		}
+
+		FRigidAssetBoneSelection NewSelection = InA;
+
+		for(const FRigidAssetBoneInfo& Bone : InB.SelectedBones)
+		{
+			NewSelection.SelectedBones.AddUnique(Bone);
+		}
+
+		NewSelection.SortBones();
+		SetValue(Context, MoveTemp(NewSelection), &Selection);
+	}
+}
+
+void FDataflowAppendBoneSelection::Register()
+{
+	AddInputs(A, B);
+	AddOutputs(Selection);
+}
+
+void FDataflowSelectConnectedBones::Evaluate(UE::Dataflow::FContext& Context, const FDataflowOutput* Out) const
+{
+	if(Out->IsA(&Selection))
+	{
+		const FRigidAssetBoneSelection& InSelection = GetValue(Context, &Selection);
+
+		if(!InSelection.Skeleton)
+		{
+			Context.Error(TEXT("Attempted to make a selection without a Skeleton"), this);
+			SafeForwardInput(Context, &Selection, &Selection);
+			return;
+		}
+
+		const FReferenceSkeleton& RefSkel = InSelection.Skeleton->GetReferenceSkeleton();
+
+		FRigidAssetBoneSelection NewSelection = InSelection;
+
+		for(const FRigidAssetBoneInfo& Bone : InSelection.SelectedBones)
+		{
+			int32 CurrentIndex = Bone.Index;
+			for(int32 UpCount = 0; UpCount < DistanceUp; ++UpCount)
+			{
+				int32 NewIndex = RefSkel.GetParentIndex(CurrentIndex);
+
+				if(NewIndex == INDEX_NONE)
+				{
+					// Nothing more in this direction
+					break;
+				}
+
+				if(!NewSelection.ContainsIndex(NewIndex))
+				{
+					NewSelection.SelectedBones.Emplace(RefSkel.GetBoneName(NewIndex), NewIndex, RefSkel.GetDepthBetweenBones(NewIndex, 0));
+				}
+
+				CurrentIndex = NewIndex;
+			}
+
+			int32 Level = 0;
+			TArray<int32> BoneStack;
+			BoneStack.Add(Bone.Index);
+
+			while(Level < DistanceDown)
+			{
+				TArray<int32> NextBoneStack;
+				while(BoneStack.Num() > 0)
+				{
+					const int32 Current = BoneStack.Pop(EAllowShrinking::No);
+
+					TArray<int32> ChildBoneIndices;
+					RefSkel.GetDirectChildBones(Current, ChildBoneIndices);
+
+					for(const int32 ChildIndex : ChildBoneIndices)
+					{
+						if(!NewSelection.ContainsIndex(ChildIndex))
+						{
+							NewSelection.SelectedBones.Emplace(RefSkel.GetBoneName(ChildIndex), ChildIndex, RefSkel.GetDepthBetweenBones(ChildIndex, 0));
+						}
+
+						NextBoneStack.Push(ChildIndex);
+					}
+				}
+				BoneStack = MoveTemp(NextBoneStack);
+
+				++Level;
+			}
+		}
+
+		NewSelection.SortBones();
+		SetValue(Context, MoveTemp(NewSelection), &Selection);
+	}
+}
+
+void FDataflowSelectConnectedBones::Register(const UE::Dataflow::FNodeParameters& InParam)
+{
+	AddInputs(Selection);
+	AddPassthroughs(Selection);
+}
+
+#if WITH_EDITOR
+bool FDataflowCreateGeometryForBones::CanDebugDraw() const
+{
+	return true;
+}
+
+bool FDataflowCreateGeometryForBones::CanDebugDrawViewMode(const FName& ViewModeName) const
+{
+	return true;
+}
+
+void FDataflowCreateGeometryForBones::DebugDraw(UE::Dataflow::FContext& Context, IDataflowDebugDrawInterface& DataflowRenderingInterface, const FDebugDrawParameters& DebugDrawParameters) const
+{
+	const FRigidAssetBoneSelection& InSelection = GetValue(Context, &Selection);
+	TRefCountPtr<IDataflowDebugDrawObject> DrawObject(MakeDebugDrawObject<FDataflowDebugDrawBoneSelectionObject>(DataflowRenderingInterface.ModifyDataflowElements(), InSelection));
+	DataflowRenderingInterface.DrawObject(DrawObject);
+}
+#endif
+
+void FDataflowCreateGeometryForBones::Evaluate(UE::Dataflow::FContext& Context, const FDataflowOutput* Out) const
+{
+	if(Out->IsA(&State))
+	{
+		TObjectPtr<UBoneGeometryGenerator> InGenerator = GetValue(Context, &Generator);
+
+		if(!InGenerator)
+		{
+			Context.Warning(TEXT("FDataflowCreateGeometryForBones: No method selected for body generation"), this);
+			SafeForwardInput(Context, &State, &State);
+			return;
+		}
+
+		const FRigidAssetBoneSelection& InSelection = GetValue(Context, &Selection);
+
+		if(!InSelection.Skeleton || !InSelection.Mesh)
+		{
+			Context.Error(TEXT("Attempted to make geometry without a skeleton or mesh"), this);
+			SafeForwardInput(Context, &State, &State);
+			return;
+		}
+
+		TObjectPtr<USkeletalBodySetup> InTemplate = GetValue(Context, &TemplateBody);
+		if(!InTemplate)
+		{
+			Context.Error(TEXT("Attempted to make a body without a template body setup"));
+			SafeForwardInput(Context, &State, &State);
+			return;
+		}
+
+		FPhysicsAssetDataflowState InState = GetValue(Context, &State);
+		TArray<UE::Chaos::RigidAsset::FSimpleGeometry> Geoms = InGenerator->Build(InSelection);
+
+		if(Geoms.Num() != InSelection.SelectedBones.Num())
+		{
+			Context.Error(TEXT("Failed to generate bodies for each bone in the selection"));
+			SafeForwardInput(Context, &State, &State);
+			return;
+		}
+
+		for(int32 Index = 0; Index < Geoms.Num(); ++Index)
+		{
+			const UE::Chaos::RigidAsset::FSimpleGeometry& Geom = Geoms[Index];
+			const FRigidAssetBoneInfo& Bone = InSelection.SelectedBones[Index];
+
+			TObjectPtr<USkeletalBodySetup> NewSetup = CastChecked<USkeletalBodySetup>(StaticDuplicateObject(InTemplate, GetTransientPackage()));
+
+			if(!Geom.GeometryElement)
+			{
+				Context.Error(TEXT("Generator failed to create usable geometry"));
+				SafeForwardInput(Context, &State, &State);
+				return;
+			}
+
+			NewSetup->AggGeom.VisitShapeAndContainer(*Geom.GeometryElement.Get(), [] <typename ElemType> (const ElemType & ElemTyped, TArray<ElemType>&Container)
+			{
+				Container.Add(ElemTyped);
+			});
+
+			NewSetup->BoneName = Bone.Name;
+
+			InState.Bodies.Add(NewSetup);
+		}
+
+		SetValue(Context, MoveTemp(InState), &State);
+	}
+}
+
+void FDataflowCreateGeometryForBones::Register()
+{
+	AddInputs(Generator, TemplateBody, State, Selection);
+	AddPassthroughs(State);
+}
+
+void FDataflowAutoConstrainBodies::Evaluate(UE::Dataflow::FContext& Context, const FDataflowOutput* Out) const
+{
+	TObjectPtr<UConstraintGenerator> InGenerator = GetValue(Context, &Generator);
+
+	if(!InGenerator)
+	{
+		Context.Warning(TEXT("FDataflowAutoConstrainBodies: No method selected for constraint generation"), this);
+		SafeForwardInput(Context, &State, &State);
+		return;
+	}
+
+	const FRigidAssetBoneSelection& InSelection = GetValue(Context, &Selection);
+
+	if(!InSelection.Skeleton || !InSelection.Mesh)
+	{
+		Context.Error(TEXT("Attempted to make constraints without a skeleton or mesh"), this);
+		SafeForwardInput(Context, &State, &State);
+		return;
+	}
+
+	TObjectPtr<UPhysicsConstraintTemplate> InTemplate = GetValue(Context, &TemplateConstraint);
+
+	if(!InTemplate)
+	{
+		Context.Error(TEXT("Attempted to make a constraint without a template"));
+		SafeForwardInput(Context, &State, &State);
+		return;
+	}
+
+	FPhysicsAssetDataflowState InState = GetValue(Context, &State);
+	TArray<TObjectPtr<UPhysicsConstraintTemplate>> GeneratedConstraints = InGenerator->Build(InTemplate, InSelection);
+
+	InState.Constraints.Append(GeneratedConstraints);
+
+	SetValue(Context, MoveTemp(InState), &State);
+}
+
+void FDataflowAutoConstrainBodies::Register()
+{
+	AddInputs(Generator, TemplateConstraint, State, Selection);
+	AddPassthroughs(State);
+}
+
+#if WITH_EDITOR
+bool FDataflowSetBodyGeometry::CanDebugDraw() const
+{
+	return true;
+}
+
+bool FDataflowSetBodyGeometry::CanDebugDrawViewMode(const FName& ViewModeName) const
+{
+	return true;
+}
+
+void FDataflowSetBodyGeometry::DebugDraw(UE::Dataflow::FContext& Context, IDataflowDebugDrawInterface& DataflowRenderingInterface, const FDebugDrawParameters& DebugDrawParameters) const
+{
+	if(DebugDrawParameters.bNodeIsSelected || DebugDrawParameters.bNodeIsPinned)
+	{
+		if(Body)
+		{
+			TRefCountPtr<IDataflowDebugDrawObject> BodyDrawObject(MakeDebugDrawObject<FDataflowDebugDrawBodySetupObject>(DataflowRenderingInterface.ModifyDataflowElements(), Body));
+			DataflowRenderingInterface.DrawObject(BodyDrawObject);
+		}
+	}
+}
+#endif
+
+void FDataflowSetBodyGeometry::Evaluate(UE::Dataflow::FContext& Context, const FDataflowOutput* Out) const
+{
+	if(Out->IsA(&Body))
+	{
+		TObjectPtr<USkeletalBodySetup> InSetup = GetValue(Context, &Body);
+
+		check(InSetup);
+
+		TObjectPtr<USkeletalBodySetup> NewSetup = CastChecked<USkeletalBodySetup>(StaticDuplicateObject(InSetup, GetTransientPackage()));
+
+		NewSetup->AggGeom = GetValue(Context, &Geometry);
+
+		SetValue(Context, NewSetup, &Body);
+	}
+}
+
+void FDataflowSetBodyGeometry::Register()
+{
+	AddInputs(Body, Geometry);
+	AddPassthroughs(Body);
+}
+
+void FDataflowNewBoneSelection::Evaluate(UE::Dataflow::FContext& Context, const FDataflowOutput* Out) const
+{
+	if(Out->IsA(&Selection))
+	{
+		TObjectPtr<USkeleton> InSkeleton = GetValue(Context, &Skeleton);
+		TObjectPtr<USkeletalMesh> InMesh = GetValue(Context, &Mesh);
+
+		// Unless overridden, get the skeleton from the currently edited asset
+		if(!InSkeleton)
+		{
+			if(TObjectPtr<UPhysicsAsset> PhysAsset = GetAttachmentOwnerAs<UPhysicsAsset>(Context))
+			{
+				InSkeleton = PhysAsset->GetPreviewMesh()->GetSkeleton();
+			}
+			else
+			{
+				Context.Error(TEXT("Attempted to make a selection without a Skeleton"), this);
+				SetValue(Context, FRigidAssetBoneSelection{}, &Selection);
+				return;
+			}
+		}
+
+		// Unless overridden, pull the skeletal mesh from the selected skeleton
+		if(!InMesh)
+		{
+			if(InSkeleton)
+			{
+				InMesh = InSkeleton->GetPreviewMesh();
+			}
+			else
+			{
+				Context.Error(TEXT("Attempted to make a selection without a Skeletal Mesh"), this);
+				SetValue(Context, FRigidAssetBoneSelection{}, &Selection);
+				return;
+			}
+		}
+
+		if(!InMesh || !InSkeleton)
+		{
+			Context.Error(TEXT("Unable to identify a mesh or skeleton"), this);
+			SetValue(Context, FRigidAssetBoneSelection{}, &Selection);
+			return;
+		}
+
+		if(InSkeleton != InMesh->GetSkeleton())
+		{
+			Context.Error(TEXT("Specified mesh is not compatible with the specified skeleton"), this);
+			SetValue(Context, FRigidAssetBoneSelection{}, &Selection);
+			return;
+		}
+
+		FRigidAssetBoneSelection NewSelection;
+		NewSelection.Skeleton = InSkeleton;
+		NewSelection.Mesh = InMesh;
+
+		SetValue(Context, MoveTemp(NewSelection), &Selection);
+	}
+}
+
+void FDataflowNewBoneSelection::Register(const UE::Dataflow::FNodeParameters& InParam)
+{
+	AddInputs(Skeleton, Mesh);
+	AddOutputs(Selection);
+}
+
+void FDataflowPhysicsAssetMakeState::Evaluate(UE::Dataflow::FContext& Context, const FDataflowOutput* Out) const
+{
+	if(Out->IsA(&State))
+	{
+		TObjectPtr<USkeletalMesh> InTarget = GetValue(Context, &TargetMesh);
+
+		// If there's no override - pick from the context asset
+		if(!InTarget)
+		{
+			if(TObjectPtr<UPhysicsAsset> PhysAsset = GetAttachmentOwnerAs<UPhysicsAsset>(Context))
+			{
+				InTarget = PhysAsset->GetPreviewMesh();
+			}
+			else
+			{
+				Context.Error(TEXT("Attempted to make a state without a valid skeleton"), this);
+				SetValue(Context, FPhysicsAssetDataflowState{}, &State);
+				return;
+			}
+		}
+
+		SetValue(Context, FPhysicsAssetDataflowState{ InTarget->GetSkeleton(), InTarget }, &State);
+	}
+}
+
+void FDataflowPhysicsAssetMakeState::Register()
+{
+	AddInputs(TargetMesh);
+	AddOutputs(State);
+}
+
+#if WITH_EDITOR
+bool FDataflowAggGeomAddShape::CanDebugDraw() const
+{
+	return true;
+}
+
+bool FDataflowAggGeomAddShape::CanDebugDrawViewMode(const FName& ViewModeName) const
+{
+	return true;
+}
+
+void FDataflowAggGeomAddShape::DebugDraw(UE::Dataflow::FContext& Context, IDataflowDebugDrawInterface& DataflowRenderingInterface, const FDebugDrawParameters& DebugDrawParameters) const
+{
+	TRefCountPtr<IDataflowDebugDrawObject> Obj(MakeDebugDrawObject<FDataflowDebugDrawAggGeomObject>(DataflowRenderingInterface.ModifyDataflowElements(), &AggGeom));
+	DataflowRenderingInterface.DrawObject(Obj);
+}
+#endif
+
+void FDataflowAggGeomAddShape::Evaluate(UE::Dataflow::FContext& Context, const FDataflowOutput* Out) const
+{
+	using namespace UE::Chaos::RigidAsset;
+
+	if(Out->IsA(&AggGeom))
+	{
+		const int32 NumShapeInputs = Shapes.Num();
+
+		const FKAggregateGeom& InGeom = GetValue(Context, &AggGeom);
+
+		FKAggregateGeom NewAggGeom = InGeom;
+
+		for(int32 Index = 0; Index < NumShapeInputs; ++Index)
+		{
+			const FSimpleGeometry& InShape = GetValue(Context, &Shapes[Index]);
+
+			if(InShape.GeometryElement)
+			{
+				NewAggGeom.AddElement(*InShape.GeometryElement.Get());
+			}
+		}
+
+		SetValue(Context, MoveTemp(NewAggGeom), &AggGeom);
+	}
+}
+
+void FDataflowAggGeomAddShape::Register()
+{
+	AddInputs(AggGeom);
+	AddPassthroughs(AggGeom);
+}
diff --git a/Engine/Plugins/Experimental/ChaosRigidAsset/Source/ChaosRigidAssetNodes/Private/PhysicsAssetDataflowNodes.h b/Engine/Plugins/Experimental/ChaosRigidAsset/Source/ChaosRigidAssetNodes/Private/PhysicsAssetDataflowNodes.h
new file mode 100644
--- /dev/null
+++ b/Engine/Plugins/Experimental/ChaosRigidAsset/Source/ChaosRigidAssetNodes/Private/PhysicsAssetDataflowNodes.h
@@ -0,0 +1,529 @@
+// Copyright Epic Games, Inc. All Rights Reserved.
+
+#pragma once
+
+#include "Concepts/ConvertibleTo.h"
+#include "PhysicsEngine/AggregateGeom.h"
+#include "PhysicsEngine/PhysicsConstraintTemplate.h"
+#include "GeometryCollection/ManagedArrayCollection.h"
+#include "GeometryCollection/ManagedArrayAccessor.h"
+#include "Dataflow/DataflowTerminalNode.h"
+
+#include "PhysicsAssetDataflowState.h"
+#include "Generators/BoneGeometryGenerators.h"
+#include "Generators/ConstraintGenerators.h"
+#include "BoneSelection.h"
+#include "RigidDataflowNode.h"
+#include "ShapeElemNodes.h"
+
+#include "PhysicsAssetDataflowNodes.generated.h"
+
+class USkeleton;
+class USkeletalMesh;
+class UPhysicsAsset;
+class USkeletalBodySetup;
+
+////////////////////////////////////////////////////////////////////////////////////////////////////////////////
+
+/**
+ * The asset state is the object that holds the intermediate data that is built throughout the execution of the
+ * dataflow graph. It's analogous to a collection used in other dataflow graphs. This graph doesn't use a
+ * collection due to needing support for copy-on-write with UObjects
+ * Asset state objects are lightweight and should be passed by-value as they support copy-on-write so each
+ * state instance is fairly small and references its shared state
+ */
+USTRUCT()
+struct FDataflowPhysicsAssetMakeState : public FRigidDataflowNode
+{
+	GENERATED_BODY()
+	DATAFLOW_NODE_DEFINE_INTERNAL(FDataflowPhysicsAssetMakeState, "PhysicsAssetMakeState", "PhysicsAsset", "")
+
+public:
+
+	FDataflowPhysicsAssetMakeState(const UE::Dataflow::FNodeParameters& InParam, FGuid InGuid = FGuid::NewGuid())
+		: FRigidDataflowNode(InParam, InGuid)
+	{
+		Register();
+	}
+
+private:
+
+	// Optional mesh input if creating a state for a mesh that isn't currently being edited
+	UPROPERTY(EditAnywhere, Category = Config, meta = (DataflowInput))
+	TObjectPtr<USkeletalMesh> TargetMesh;
+
+	// Resulting state for the mesh
+	UPROPERTY(meta = (DataflowOutput))
+	FPhysicsAssetDataflowState State;
+
+	void Evaluate(UE::Dataflow::FContext& Context, const FDataflowOutput* Out) const override;
+	void Register();
+};
+
+////////////////////////////////////////////////////////////////////////////////////////////////////////////////
+
+/**
+ * Add a shape to an aggregate geometry
+ */
+USTRUCT()
+struct FDataflowAggGeomAddShape : public FRigidDataflowMultiInputNode
+{
+	GENERATED_BODY()
+	DATAFLOW_NODE_DEFINE_INTERNAL(FDataflowAggGeomAddShape, "AggGeomAddShape", "PhysicsAsset", "")
+	DATAFLOW_NODE_RENDER_TYPE("GeomRender", FName("FKAggregateGeom"), "AggGeom")
+
+public:
+
+	FDataflowAggGeomAddShape(const UE::Dataflow::FNodeParameters& InParam, FGuid InGuid = FGuid::NewGuid())
+		: FRigidDataflowMultiInputNode(Shapes, InParam, InGuid)
+	{
+		InitPins();
+		Register();
+	}
+
+private:
+
+	// Shapes to add to the aggregate geometry
+	UPROPERTY()
+	TArray<UE::Chaos::RigidAsset::FSimpleGeometry> Shapes;
+
+	// The aggregate geometry to add the shape to
+	UPROPERTY(meta = (DataflowInput, DataflowOutput, DataflowPassthrough = AggGeom))
+	FKAggregateGeom AggGeom;
+
+#if WITH_EDITOR
+	bool CanDebugDraw() const override;
+	bool CanDebugDrawViewMode(const FName& ViewModeName) const override;
+	void DebugDraw(UE::Dataflow::FContext& Context, IDataflowDebugDrawInterface& DataflowRenderingInterface, const FDebugDrawParameters& DebugDrawParameters) const override;
+#endif
+
+
+	void Evaluate(UE::Dataflow::FContext& Context, const FDataflowOutput* Out) const override;
+	void Register();
+};
+
+template<>
+struct TStructOpsTypeTraits<FDataflowAggGeomAddShape> : public TStructOpsTypeTraitsBase2<FDataflowAggGeomAddShape>
+{
+	enum
+	{
+		WithCopy = false
+	};
+};
+
+////////////////////////////////////////////////////////////////////////////////////////////////////////////////
+
+/**
+ * Terminal node that converts an asset state into the final asset
+ */
+USTRUCT(meta = (DataflowTerminal))
+struct FDataflowPhysicsAssetTerminalNode : public FDataflowTerminalNode
+{
+	GENERATED_BODY()
+	DATAFLOW_NODE_DEFINE_INTERNAL(FDataflowPhysicsAssetTerminalNode, "PhysicsAssetTerminal", "Terminal", "")
+
+public:
+
+	FDataflowPhysicsAssetTerminalNode(const UE::Dataflow::FNodeParameters& InParam, FGuid InGuid = FGuid::NewGuid());
+
+private:
+
+	// The state to convert into the final asset
+	UPROPERTY(meta = (DataflowInput))
+	FPhysicsAssetDataflowState State;
+
+	// Override target asset to write to - optional, defaults to the current asset in the editor
+	UPROPERTY(EditAnywhere, Category = Asset, meta = (DataflowInput))
+	TObjectPtr<UPhysicsAsset> PhysicsAsset = nullptr;
+
+	// FDataflowTerminalNode
+	void SetAssetValue(TObjectPtr<UObject> Asset, UE::Dataflow::FContext& Context) const override;
+	void Evaluate(UE::Dataflow::FContext& Context) const override;
+	////////////////////////////////////////////////////////////////////////////////////////////////////////////////
+};
+
+////////////////////////////////////////////////////////////////////////////////////////////////////////////////
+
+/**
+ * Add a body setup to an asset state
+ */
+USTRUCT()
+struct FDataflowPhysicsAssetAddBody : public FRigidDataflowNode
+{
+	GENERATED_BODY()
+	DATAFLOW_NODE_DEFINE_INTERNAL(FDataflowPhysicsAssetAddBody, "PhysicsAssetAddBody", "PhysicsAsset", "")
+
+public:
+
+	FDataflowPhysicsAssetAddBody(const UE::Dataflow::FNodeParameters& InParam, FGuid InGuid = FGuid::NewGuid());
+
+private:
+
+	// Body to add
+	UPROPERTY(meta = (DataflowInput))
+	TObjectPtr<USkeletalBodySetup> Body;
+
+	// State to write to
+	UPROPERTY(meta = (DataflowInput, DataflowOutput, DataflowPassthrough = State))
+	FPhysicsAssetDataflowState State;
+
+	void Evaluate(UE::Dataflow::FContext& Context, const FDataflowOutput* Out) const override;
+};
+
+////////////////////////////////////////////////////////////////////////////////////////////////////////////////
+
+/**
+ * Make a single body setup from a template
+ */
+USTRUCT()
+struct FDataflowMakeBody : public FRigidDataflowNode
+{
+	GENERATED_BODY()
+	DATAFLOW_NODE_DEFINE_INTERNAL(FDataflowMakeBody, "Make Body Setup", "PhysicsAsset", "")
+
+public:
+
+	FDataflowMakeBody(const UE::Dataflow::FNodeParameters& InParam, FGuid InGuid = FGuid::NewGuid())
+		: FRigidDataflowNode(InParam, InGuid)
+	{
+		Register(InParam);
+	}
+
+private:
+
+#if WITH_EDITOR
+	bool CanDebugDraw() const override;
+	bool CanDebugDrawViewMode(const FName& ViewModeName) const override;
+	void DebugDraw(UE::Dataflow::FContext& Context, IDataflowDebugDrawInterface& DataflowRenderingInterface, const FDebugDrawParameters& DebugDrawParameters) const override;
+#endif
+
+
+	// Template body setup for details panel and output new body instance
+	UPROPERTY(EditAnywhere, Instanced, Category = Body, meta = (DataflowOutput))
+	TObjectPtr<USkeletalBodySetup> Body;
+
+	// Bone that this body should be bound to
+	UPROPERTY(EditAnywhere, Category = Body, meta = (DataflowInput))
+	FName BoneName;
+
+	void Evaluate(UE::Dataflow::FContext& Context, const FDataflowOutput* Out) const override;
+	void Register(const UE::Dataflow::FNodeParameters& InParam);
+};
+
+////////////////////////////////////////////////////////////////////////////////////////////////////////////////
+
+/**
+ * Make a single joint setup from a template
+ */
+USTRUCT()
+struct FDataflowMakeJoint : public FRigidDataflowNode
+{
+	GENERATED_BODY()
+	DATAFLOW_NODE_DEFINE_INTERNAL(FDataflowMakeJoint, "Make Joint", "PhysicsAsset", "")
+
+public:
+
+	FDataflowMakeJoint(const UE::Dataflow::FNodeParameters& InParam, FGuid InGuid = FGuid::NewGuid())
+		: FRigidDataflowNode(InParam, InGuid)
+	{
+		Register(InParam);
+	}
+
+private:
+
+	UPROPERTY(EditAnywhere, Instanced, Category = Body, meta = (DataflowOutput))
+	TObjectPtr<UPhysicsConstraintTemplate> Joint;
+
+	void Evaluate(UE::Dataflow::FContext& Context, const FDataflowOutput* Out) const override;
+	void Register(const UE::Dataflow::FNodeParameters& InParam);
+};
+
+////////////////////////////////////////////////////////////////////////////////////////////////////////////////
+
+/**
+ * Select bones in a mesh given a name search pattern
+ */
+USTRUCT()
+struct FDataflowSelectBonesByName : public FRigidDataflowNode
+{
+	GENERATED_BODY()
+	DATAFLOW_NODE_DEFINE_INTERNAL(FDataflowSelectBonesByName, "Select Bones by Name", "PhysicsAsset", "")
+	DATAFLOW_NODE_RENDER_TYPE("GeomRender", FName("FBoneSelection"), "Selection")
+
+public:
+
+	FDataflowSelectBonesByName(const UE::Dataflow::FNodeParameters& InParam, FGuid InGuid = FGuid::NewGuid())
+		: FRigidDataflowNode(InParam, InGuid)
+	{
+		Register(InParam);
+	}
+
+private:
+
+	// State to select bones within
+	UPROPERTY(meta = (DataflowInput))
+	FPhysicsAssetDataflowState State;
+
+	// Resulting selection
+	UPROPERTY(meta = (DataflowOutput))
+	FRigidAssetBoneSelection Selection;
+
+	// Search pattern to use. Supports wildcards (? - match single character, * - match sequence)
+	UPROPERTY(EditAnywhere, Category = Parameters)
+	FString SearchString;
+
+	void Evaluate(UE::Dataflow::FContext& Context, const FDataflowOutput* Out) const override;
+	void Register(const UE::Dataflow::FNodeParameters& InParam);
+};
+
+////////////////////////////////////////////////////////////////////////////////////////////////////////////////
+
+/**
+ * Append a bone selection to another
+ */
+USTRUCT()
+struct FDataflowAppendBoneSelection : public FRigidDataflowNode
+{
+	GENERATED_BODY()
+	DATAFLOW_NODE_DEFINE_INTERNAL(FDataflowAppendBoneSelection, "Append Selected Bones", "PhysicsAsset", "")
+	DATAFLOW_NODE_RENDER_TYPE("GeomRender", FName("FBoneSelection"), "Selection")
+
+public:
+
+	FDataflowAppendBoneSelection(const UE::Dataflow::FNodeParameters& InParam, FGuid InGuid = FGuid::NewGuid())
+		: FRigidDataflowNode(InParam, InGuid)
+	{
+		Register();
+	}
+
+private:
+
+	// First selection
+	UPROPERTY(meta = (DataflowInput))
+	FRigidAssetBoneSelection A;
+
+	// Second selection
+	UPROPERTY(meta = (DataflowInput))
+	FRigidAssetBoneSelection B;
+
+	// Resulting selection, a combination of A and B
+	UPROPERTY(meta = (DatflowOutput))
+	FRigidAssetBoneSelection Selection;
+
+	void Evaluate(UE::Dataflow::FContext& Context, const FDataflowOutput* Out) const override;
+	void Register();
+};
+
+////////////////////////////////////////////////////////////////////////////////////////////////////////////////
+
+/**
+ * Select bones connected to each bone already in a selection.
+ */
+USTRUCT()
+struct FDataflowSelectConnectedBones : public FRigidDataflowNode
+{
+	GENERATED_BODY()
+	DATAFLOW_NODE_DEFINE_INTERNAL(FDataflowSelectConnectedBones, "Select Connected Bones", "PhysicsAsset", "")
+	DATAFLOW_NODE_RENDER_TYPE("GeomRender", FName("FBoneSelection"), "Selection")
+
+public:
+
+	FDataflowSelectConnectedBones(const UE::Dataflow::FNodeParameters& InParam, FGuid InGuid = FGuid::NewGuid())
+		: FRigidDataflowNode(InParam, InGuid)
+	{
+		Register(InParam);
+	}
+
+private:
+
+	// Disance to walk up the hierarchy to find bones
+	UPROPERTY(EditAnywhere, Category = Settings)
+	int32 DistanceUp = 0;
+
+	// Distance to walk down the hierarchy to find bones
+	UPROPERTY(EditAnywhere, Category = Settings)
+	int32 DistanceDown = 0;
+
+	// Resulting bone selction
+	UPROPERTY(meta = (DataflowInput, DataflowOutput, DataflowPassthrough = "Selection"))
+	FRigidAssetBoneSelection Selection;
+
+	void Evaluate(UE::Dataflow::FContext& Context, const FDataflowOutput* Out) const override;
+	void Register(const UE::Dataflow::FNodeParameters& InParam);
+};
+
+////////////////////////////////////////////////////////////////////////////////////////////////////////////////
+
+/**
+ * Apply a geometry generator to a bone selection, creating bodies for each bone
+ * The method of geometry generation is generic and offloaded to a "Generator" object/
+ * @see UBoneGeometryGenerator for implementation and creating your own generator node.
+ */
+USTRUCT()
+struct FDataflowCreateGeometryForBones : public FRigidDataflowNode
+{
+	GENERATED_BODY()
+	DATAFLOW_NODE_DEFINE_INTERNAL(FDataflowCreateGeometryForBones, "Create Geometry for Bones", "PhysicsAsset", "")
+	DATAFLOW_NODE_RENDER_TYPE("GeomRender", FName("FPhysicsAssetDataflowState"), "State")
+
+public:
+
+	FDataflowCreateGeometryForBones(const UE::Dataflow::FNodeParameters& InParam, FGuid InGuid = FGuid::NewGuid())
+		: FRigidDataflowNode(InParam, InGuid)
+	{
+		Register();
+	}
+
+private:
+
+	UPROPERTY(meta=(DataflowInput))
+	TObjectPtr<UBoneGeometryGenerator> Generator;
+
+	UPROPERTY(meta=(DataflowInput))
+	TObjectPtr<USkeletalBodySetup> TemplateBody;
+
+	UPROPERTY(EditAnywhere, Category=Generation)
+	bool bCreateMissingBodies = true;
+
+	UPROPERTY(meta = (DataflowInput, DataflowOutput, DataflowPassthrough = "State"))
+	FPhysicsAssetDataflowState State;
+
+	UPROPERTY(meta = (DataflowInput))
+	FRigidAssetBoneSelection Selection;
+
+#if WITH_EDITOR
+	bool CanDebugDraw() const override;
+	bool CanDebugDrawViewMode(const FName& ViewModeName) const override;
+	void DebugDraw(UE::Dataflow::FContext& Context, IDataflowDebugDrawInterface& DataflowRenderingInterface, const FDebugDrawParameters& DebugDrawParameters) const override;
+#endif
+
+	void Evaluate(UE::Dataflow::FContext& Context, const FDataflowOutput* Out) const override;
+	void Register();
+};
+
+////////////////////////////////////////////////////////////////////////////////////////////////////////////////
+
+/**
+ * Apply a constraint generator to a bone selection, creating constraints for each pair of bodies
+ * The method of constraint generation is generic and offloaded to a "Generator" object/
+ * @see UConstraintGenerator for implementation and creating your own generator node.
+ */
+USTRUCT()
+struct FDataflowAutoConstrainBodies : public FRigidDataflowNode
+{
+	GENERATED_BODY();
+	DATAFLOW_NODE_DEFINE_INTERNAL(FDataflowAutoConstrainBodies, "Create Constraints for Bones", "PhysicsAsset", "")
+
+public:
+
+	FDataflowAutoConstrainBodies(const UE::Dataflow::FNodeParameters& InParam, FGuid InGuid = FGuid::NewGuid())
+		: FRigidDataflowNode(InParam, InGuid)
+	{
+		Register();
+	}
+
+private:
+
+	// Generation method to use
+	UPROPERTY(meta=(DataflowInput))
+	TObjectPtr<UConstraintGenerator> Generator;
+
+	// Template constraint to take properties from
+	UPROPERTY(meta=(DataflowInput))
+	TObjectPtr<UPhysicsConstraintTemplate> TemplateConstraint;
+
+	// State to apply the joints to
+	UPROPERTY(meta = (DataflowInput, DataflowOutput, DataflowPassthrough = "State"))
+	FPhysicsAssetDataflowState State;
+
+	// Optional selection, restricts the bodies that will attempt to auto constrain
+	UPROPERTY(meta = (DataflowInput))
+	FRigidAssetBoneSelection Selection;
+
+	void Evaluate(UE::Dataflow::FContext& Context, const FDataflowOutput* Out) const override;
+	void Register();
+};
+
+////////////////////////////////////////////////////////////////////////////////////////////////////////////////
+
+/**
+ * Given a body setup, replace the internal aggregate geometry with the provided geometry
+ */
+USTRUCT()
+struct FDataflowSetBodyGeometry : public FRigidDataflowNode
+{
+	GENERATED_BODY()
+	DATAFLOW_NODE_DEFINE_INTERNAL(FDataflowSetBodyGeometry, "Set Body Geometry", "PhysicsAsset", "")
+
+public:
+
+	FDataflowSetBodyGeometry(const UE::Dataflow::FNodeParameters& InParam, FGuid InGuid = FGuid::NewGuid())
+		: FRigidDataflowNode(InParam, InGuid)
+	{
+		Register();
+	}
+
+private:
+
+	// The body to modify
+	UPROPERTY(meta = (DataflowInput, DataflowOutput, DataflowPassthrough = "Body"))
+	TObjectPtr<USkeletalBodySetup> Body;
+
+	// The Geometry to set on the body
+	UPROPERTY(meta = (DataflowInput))
+	FKAggregateGeom Geometry;
+
+#if WITH_EDITOR
+	bool CanDebugDraw() const override;
+	bool CanDebugDrawViewMode(const FName& ViewModeName) const override;
+	void DebugDraw(UE::Dataflow::FContext& Context, IDataflowDebugDrawInterface& DataflowRenderingInterface, const FDebugDrawParameters& DebugDrawParameters) const override;
+#endif
+
+	void Evaluate(UE::Dataflow::FContext& Context, const FDataflowOutput* Out) const override;
+	void Register();
+};
+
+////////////////////////////////////////////////////////////////////////////////////////////////////////////////
+
+/**
+ * Create a new, empty selection for the specified mesh
+ */
+USTRUCT()
+struct FDataflowNewBoneSelection : public FRigidDataflowNode
+{
+	GENERATED_BODY()
+	DATAFLOW_NODE_DEFINE_INTERNAL(FDataflowNewBoneSelection, "New Bone Selection", "PhysicsAsset", "")
+
+public:
+
+	FDataflowNewBoneSelection(const UE::Dataflow::FNodeParameters& InParam, FGuid InGuid = FGuid::NewGuid())
+		: FRigidDataflowNode(InParam, InGuid)
+	{
+		Register(InParam);
+	}
+
+private:
+
+	// The created selection
+	UPROPERTY(meta = (DataflowOutput))
+	FRigidAssetBoneSelection Selection;
+
+	// The skeleton to target if not using the currently edited asset
+	UPROPERTY(meta = (DataflowInput))
+	TObjectPtr<USkeleton> Skeleton = nullptr;
+
+	// The Mesh to target if not using the currently edited asset
+	UPROPERTY(meta = (DataflowInput))
+	TObjectPtr<USkeletalMesh> Mesh = nullptr;
+
+	void Evaluate(UE::Dataflow::FContext& Context, const FDataflowOutput* Out) const override;
+	void Register(const UE::Dataflow::FNodeParameters& InParam);
+};
+
+////////////////////////////////////////////////////////////////////////////////////////////////////////////////
+
+namespace UE::Dataflow
+{
+	void RegisterPhysicsAssetTerminalNode();
+	void RegisterPhysicsAssetNodes();
+}
diff --git a/Engine/Plugins/Experimental/ChaosRigidAsset/Source/ChaosRigidAssetNodes/Private/PhysicsAssetDataflowState.cpp b/Engine/Plugins/Experimental/ChaosRigidAsset/Source/ChaosRigidAssetNodes/Private/PhysicsAssetDataflowState.cpp
new file mode 100644
--- /dev/null
+++ b/Engine/Plugins/Experimental/ChaosRigidAsset/Source/ChaosRigidAssetNodes/Private/PhysicsAssetDataflowState.cpp
@@ -0,0 +1,29 @@
+// Copyright Epic Games, Inc. All Rights Reserved.
+
+#include "PhysicsAssetDataflowState.h"
+#include "PhysicsEngine/SkeletalBodySetup.h"
+#include "PhysicsEngine/PhysicsConstraintTemplate.h"
+
+void FPhysicsAssetDataflowState::AddStructReferencedObjects(FReferenceCollector& Collector)
+{
+	Collector.AddReferencedObject(TargetSkeleton);
+	AddArrayReferences(Collector, Bodies, Constraints);
+}
+
+void FPhysicsAssetDataflowState::DebugLog() const
+{
+	for(const USkeletalBodySetup* Setup : Bodies)
+	{
+		UE_LOG(LogTemp, Warning, TEXT("Body: %s"), *Setup->GetName());
+	}
+
+	for(const UPhysicsConstraintTemplate* Constraint : Constraints)
+	{
+		UE_LOG(LogTemp, Warning, TEXT("Constraint: %s"), *Constraint->GetName());
+	}
+}
+
+bool FPhysicsAssetDataflowState::HasData() const
+{
+	return TargetSkeleton && TargetMesh;
+}
diff --git a/Engine/Plugins/Experimental/ChaosRigidAsset/Source/ChaosRigidAssetNodes/Private/RigidDataflowNode.h b/Engine/Plugins/Experimental/ChaosRigidAsset/Source/ChaosRigidAssetNodes/Private/RigidDataflowNode.h
new file mode 100644
--- /dev/null
+++ b/Engine/Plugins/Experimental/ChaosRigidAsset/Source/ChaosRigidAssetNodes/Private/RigidDataflowNode.h
@@ -0,0 +1,343 @@
+// Copyright Epic Games, Inc. All Rights Reserved.
+
+#pragma once
+
+#include "Dataflow/DataflowNode.h"
+#include "Dataflow/DataflowObjectInterface.h"
+#include "ShapeElemNodes.h"
+#include "DataflowAttachment.h"
+
+#include "RigidDataflowNode.generated.h"
+
+/**
+ * Intermediate node for rigid asset node implementations encapsulating common functionality
+ */
+USTRUCT()
+struct FRigidDataflowNode : public FDataflowNode
+{
+	GENERATED_BODY();
+
+public:
+
+	FRigidDataflowNode() = default;
+
+	FRigidDataflowNode(const UE::Dataflow::FNodeParameters& Param, FGuid InGuid)
+		: FDataflowNode(Param, InGuid)
+		, Owner(Param.OwningObject)
+	{
+
+	}
+
+	// Get the owning UObject for the node, to store UObjects in the asset safely
+	TObjectPtr<UObject> GetOwner() const
+	{
+		return Owner;
+	}
+
+	// Multi-input versions of the pin registration functions
+	template<typename... T>
+	void AddInputs(T&&... Args)
+	{
+		(RegisterInputConnection(&Args), ...);
+	}
+
+	template<typename... T>
+	void AddOutputs(T&&... Args)
+	{
+		(RegisterOutputConnection(&Args), ...);
+	}
+
+	template<typename... T>
+	void AddPassthroughs(T&&... Args)
+	{
+		(RegisterOutputConnection(&Args, &Args), ...);
+	}
+
+protected:
+
+	// Grab the context owner and return it as T if possible, otherwise nullptr
+	template<typename T>
+	TObjectPtr<T> GetContextOwnerAs(UE::Dataflow::FContext& InContext) const
+	{
+		if(UE::Dataflow::FEngineContext* EngineContext = InContext.AsType<UE::Dataflow::FEngineContext>())
+		{
+			return Cast<T>(EngineContext->Owner);
+		}
+
+		return nullptr;
+	}
+
+	TObjectPtr<UObject> GetContextOwner(UE::Dataflow::FContext& InContext) const
+	{
+		if(UE::Dataflow::FEngineContext* EngineContext = InContext.AsType<UE::Dataflow::FEngineContext>())
+		{
+			return EngineContext->Owner;
+		}
+
+		return nullptr;
+	}
+
+	// Get the dataflow attachement if it is avaialble for the current asset
+	TObjectPtr<UDataflowAttachment> GetAttachment(UE::Dataflow::FContext& InContext) const
+	{
+		return GetContextOwnerAs<UDataflowAttachment>(InContext);
+	}
+
+	TObjectPtr<UObject> GetAttachmentOwner(UE::Dataflow::FContext& InContext) const
+	{
+		if(TObjectPtr<UDataflowAttachment> Attachment = GetAttachment(InContext))
+		{
+			return Attachment->GetOuter();
+		}
+
+		return nullptr;
+	}
+
+	template<typename T>
+	TObjectPtr<T> GetAttachmentOwnerAs(UE::Dataflow::FContext& InContext) const
+	{
+		if(TObjectPtr<UDataflowAttachment> Attachment = GetAttachment(InContext))
+		{
+			return Cast<T>(Attachment->GetOuter());
+		}
+
+		return nullptr;
+	}
+
+private:
+
+	UPROPERTY()
+	TObjectPtr<UObject> Owner;
+
+};
+
+struct IMultiPinArrayConnector
+{
+	virtual ~IMultiPinArrayConnector()
+	{
+	}
+
+	virtual void InitPins() = 0;
+	virtual TArray<UE::Dataflow::FPin> AddPins() = 0;
+	virtual bool CanAddPin() const = 0;
+	virtual TArray<UE::Dataflow::FPin> GetPinsToRemove() const = 0;
+	virtual void OnPinRemoved(const UE::Dataflow::FPin& Pin) = 0;
+	virtual bool CanRemovePin() const = 0;
+	virtual void PostSerialize(const FArchive& Ar) = 0;
+};
+
+template<typename T>
+struct TMultiPinArrayConnector : public IMultiPinArrayConnector
+{
+	TMultiPinArrayConnector() = delete;
+	explicit TMultiPinArrayConnector(FRigidDataflowNode& InNode, TArray<T>& InPinContainer, int32 InMinimumNumPins)
+		: Node(InNode)
+		, PinContainer(InPinContainer)
+		, MinimumNumPins(InMinimumNumPins)
+	{
+		
+	}
+
+	void InitPins() override
+	{
+		check(PinContainer.Num() == 0);
+
+		for(int32 Index = 0; Index < MinimumNumPins; ++Index)
+		{
+			AddPins();
+		}
+	}
+
+	TArray<UE::Dataflow::FPin> AddPins() override
+	{
+		const int32 Index = PinContainer.AddDefaulted();
+		UE::Dataflow::TConnectionReference<UE::Chaos::RigidAsset::FSimpleGeometry> Ref{ &PinContainer[Index], Index, &PinContainer };
+		const FDataflowInput& Input = Node.RegisterInputArrayConnection(Ref);
+
+		return
+		{
+			{
+				.Direction = UE::Dataflow::FPin::EDirection::INPUT,
+				.Type = Input.GetType(),
+				.Name = Input.GetName()
+			}
+		};
+	}
+
+
+	bool CanAddPin() const override
+	{
+		return true;
+	}
+
+
+	TArray<UE::Dataflow::FPin> GetPinsToRemove() const override
+	{
+		const int32 Index = PinContainer.Num() - 1;
+		UE::Dataflow::TConnectionReference<UE::Chaos::RigidAsset::FSimpleGeometry> Ref{ &PinContainer[Index], Index, &PinContainer };
+
+		if(const FDataflowInput* Input = Node.FindInput(Ref))
+		{
+			return
+			{
+				{
+					.Direction = UE::Dataflow::FPin::EDirection::INPUT,
+					.Type = Input->GetType(),
+					.Name = Input->GetName()
+				}
+			};
+		}
+
+		return {};
+	}
+
+
+	void OnPinRemoved(const UE::Dataflow::FPin& Pin) override
+	{
+		PinContainer.Pop();
+	}
+
+	bool CanRemovePin() const override
+	{
+		return PinContainer.Num() > MinimumNumPins;
+	}
+	
+	void PostSerialize(const FArchive& Ar) override
+	{
+		// Restore pins
+		if(Ar.IsLoading())
+		{
+			for(int32 Index = 0; Index < PinContainer.Num(); ++Index)
+			{
+				UE::Dataflow::TConnectionReference<UE::Chaos::RigidAsset::FSimpleGeometry> Ref{ &PinContainer[Index], Index, &PinContainer };
+				Node.FindOrRegisterInputArrayConnection(Ref);
+			}
+
+			if(Ar.IsTransacting())
+			{
+				int32 OriginalNum = Node.GetNumInputs();
+				const int32 ExpectedNum = PinContainer.Num();
+
+				if(OriginalNum > ExpectedNum)
+				{
+					PinContainer.SetNum(OriginalNum);
+
+					while(OriginalNum-- > ExpectedNum)
+					{
+						const int32 Index = PinContainer.Num() - 1;
+						UE::Dataflow::TConnectionReference<UE::Chaos::RigidAsset::FSimpleGeometry> Ref{ &PinContainer[Index], Index, &PinContainer };
+						Node.UnregisterInputConnection(Ref);
+
+						PinContainer.Pop(EAllowShrinking::No);
+					}
+
+					check(PinContainer.Num() == ExpectedNum);
+				}
+			}
+		}
+	}
+
+private:
+
+	FRigidDataflowNode& Node;
+	TArray<T>& PinContainer;
+	int32 MinimumNumPins;
+};
+
+USTRUCT()
+struct FRigidDataflowMultiInputNode : public FRigidDataflowNode
+{
+	GENERATED_BODY()
+
+	FRigidDataflowMultiInputNode() = default;
+
+	template<typename T>
+	explicit FRigidDataflowMultiInputNode(TArray<T>& InPinContainer, const UE::Dataflow::FNodeParameters& InParam, FGuid InGuid)
+		: FRigidDataflowNode(InParam, InGuid)
+	{
+		Connector = MakeUnique<TMultiPinArrayConnector<T>>(*this, InPinContainer, 0);
+	}
+
+	void InitPins()
+	{
+		if(Connector)
+		{
+			Connector->InitPins();
+		}
+	}
+
+	TArray<UE::Dataflow::FPin> AddPins() override
+	{
+		if(Connector)
+		{
+			return Connector->AddPins();
+		}
+
+		return Super::AddPins();
+	}
+
+	bool CanAddPin() const override
+	{
+		if(Connector)
+		{
+			return Connector->CanAddPin();
+		}
+
+		return Super::CanAddPin();
+	}
+
+
+	TArray<UE::Dataflow::FPin> GetPinsToRemove() const override
+	{
+		if(Connector)
+		{
+			return Connector->GetPinsToRemove();
+		}
+
+		return Super::GetPinsToRemove();
+	}
+
+
+	void OnPinRemoved(const UE::Dataflow::FPin& Pin) override
+	{
+		if(Connector)
+		{
+			Connector->OnPinRemoved(Pin);
+			return;
+		}
+
+		Super::OnPinRemoved(Pin);
+	}
+
+	bool CanRemovePin() const override
+	{
+		if(Connector)
+		{
+			return Connector->CanRemovePin();
+		}
+
+		return Super::CanRemovePin();
+	}
+
+	void PostSerialize(const FArchive& Ar) override
+	{
+		if(Connector)
+		{
+			Connector->PostSerialize(Ar);
+		}
+
+		Super::PostSerialize(Ar);
+	}
+
+private:
+	TUniquePtr<IMultiPinArrayConnector> Connector = nullptr;
+};
+
+template<>
+struct TStructOpsTypeTraits<FRigidDataflowMultiInputNode> : public TStructOpsTypeTraitsBase2<FRigidDataflowMultiInputNode>
+{
+	enum
+	{
+		WithCopy = false
+	};
+};
\ No newline at end of file
diff --git a/Engine/Plugins/Experimental/ChaosRigidAsset/Source/ChaosRigidAssetNodes/Private/ShapeElemNodes.cpp b/Engine/Plugins/Experimental/ChaosRigidAsset/Source/ChaosRigidAssetNodes/Private/ShapeElemNodes.cpp
new file mode 100644
--- /dev/null
+++ b/Engine/Plugins/Experimental/ChaosRigidAsset/Source/ChaosRigidAssetNodes/Private/ShapeElemNodes.cpp
@@ -0,0 +1,97 @@
+// Copyright Epic Games, Inc. All Rights Reserved.
+
+#include "ShapeElemNodes.h"
+#include "Dataflow/DataflowNode.h"
+#include "Dataflow/DataflowNodeFactory.h"
+#include "PhysicsEngine/ShapeElem.h"
+#include "PhysicsEngine/ConvexElem.h"
+
+namespace UE::Dataflow
+{
+	void RegisterShapeNodes()
+	{
+		using namespace UE::Chaos::RigidAsset;
+
+		DATAFLOW_NODE_REGISTER_CREATION_FACTORY(FMakeBoxElemDataflowNode);
+		DATAFLOW_NODE_REGISTER_CREATION_FACTORY(FMakeSphereElemDataflowNode);
+		DATAFLOW_NODE_REGISTER_CREATION_FACTORY(FMakeCapsuleElemDataflowNode);
+	}
+}
+
+namespace UE::Chaos::RigidAsset
+{
+	void FMakeBoxElemDataflowNode::Evaluate(UE::Dataflow::FContext& Context, const FDataflowOutput* Out) const
+	{
+		FSimpleGeometry NewGeom;
+		NewGeom.GeometryElement = MakeShared<FKBoxElem>(Box);
+
+		SetValue(Context, MoveTemp(NewGeom), &Geometry);
+	}
+
+	void FMakeBoxElemDataflowNode::Register()
+	{
+		RegisterOutputConnection(&Geometry);
+	}
+
+	// Helper to avoid handling every access in a switch statement
+	// TODO move to Engine
+	template<typename Callable>
+	auto VisitAsConcreteElement(FKShapeElem* InElem, Callable&& InCallable)
+	{
+		switch(InElem->GetShapeType())
+		{
+		case EAggCollisionShape::Sphere:
+			return InCallable(InElem->GetShapeCheck<FKSphereElem>());
+		case EAggCollisionShape::Box:
+			return InCallable(InElem->GetShapeCheck<FKBoxElem>());
+		case EAggCollisionShape::Sphyl:
+			return InCallable(InElem->GetShapeCheck<FKSphylElem>());
+		case EAggCollisionShape::Convex:
+			return InCallable(InElem->GetShapeCheck<FKConvexElem>());
+		default:
+			UE_LOG(LogTemp, Warning, TEXT("Unsupported rigid shape type"));
+		}
+
+		return InCallable(static_cast<FKShapeElem*>(nullptr));
+	}
+
+	FSimpleGeometry FSimpleGeometry::Duplicate()
+	{
+		TSharedPtr<FKShapeElem> NewElem;
+		
+		VisitAsConcreteElement(GeometryElement.Get(), [&NewElem] <typename ElemType> (ElemType* Elem)
+			{
+				check(Elem);
+
+				NewElem = MakeShared<ElemType>(*Elem);
+			});
+
+		return { LocalTransform, MoveTemp(NewElem) };
+	}
+
+	void FMakeSphereElemDataflowNode::Evaluate(UE::Dataflow::FContext& Context, const FDataflowOutput* Out) const
+	{
+		FSimpleGeometry NewGeom;
+		NewGeom.GeometryElement = MakeShared<FKSphereElem>(Sphere);
+
+		SetValue(Context, MoveTemp(NewGeom), &Geometry);
+	}
+
+	void FMakeSphereElemDataflowNode::Register()
+	{
+		RegisterOutputConnection(&Geometry);
+	}
+
+	void FMakeCapsuleElemDataflowNode::Evaluate(UE::Dataflow::FContext& Context, const FDataflowOutput* Out) const
+	{
+		FSimpleGeometry NewGeom;
+		NewGeom.GeometryElement = MakeShared<FKSphylElem>(Capsule);
+
+		SetValue(Context, MoveTemp(NewGeom), &Geometry);
+	}
+
+	void FMakeCapsuleElemDataflowNode::Register()
+	{
+		RegisterOutputConnection(&Geometry);
+	}
+}
diff --git a/Engine/Plugins/Experimental/ChaosRigidAsset/Source/ChaosRigidAssetNodes/Private/ShapeElemNodes.h b/Engine/Plugins/Experimental/ChaosRigidAsset/Source/ChaosRigidAssetNodes/Private/ShapeElemNodes.h
new file mode 100644
--- /dev/null
+++ b/Engine/Plugins/Experimental/ChaosRigidAsset/Source/ChaosRigidAssetNodes/Private/ShapeElemNodes.h
@@ -0,0 +1,130 @@
+// Copyright Epic Games, Inc. All Rights Reserved.
+
+#pragma once
+
+#include "PhysicsEngine/BoxElem.h"
+#include "PhysicsEngine/SphereElem.h"
+#include "PhysicsEngine/SphylElem.h"
+#include "Dataflow/DataflowNode.h"
+
+#include "ShapeElemNodes.generated.h"
+
+namespace UE::Dataflow
+{
+	void RegisterShapeNodes();
+}
+
+namespace UE::Chaos::RigidAsset
+{
+	////////////////////////////////////////////////////////////////////////////////////////////////////////////////
+
+	/**
+	 * Generic container for any geometry element while referenced in a dataflow graph
+	 */
+	USTRUCT()
+	struct FSimpleGeometry
+	{
+		GENERATED_BODY();
+
+		// Duplicate this geometry handling the geometry type correctly
+		FSimpleGeometry Duplicate();
+
+		FTransform LocalTransform = FTransform::Identity;
+		TSharedPtr<FKShapeElem> GeometryElement = nullptr;
+	};
+
+	////////////////////////////////////////////////////////////////////////////////////////////////////////////////
+
+	/**
+	 * Make single-element nodes.
+	 */
+	USTRUCT()
+	struct FMakeBoxElemDataflowNode : public FDataflowNode
+	{
+		GENERATED_BODY()
+		DATAFLOW_NODE_DEFINE_INTERNAL(FMakeBoxElemDataflowNode, "MakeBox", "Rigid Shapes", "")
+	
+	public:
+	
+		FMakeBoxElemDataflowNode(const UE::Dataflow::FNodeParameters& InParam, FGuid InGuid = FGuid::NewGuid())
+			: FDataflowNode(InParam, InGuid)
+		{
+			Register();
+		}
+
+	private:
+
+		// Input box data
+		UPROPERTY(EditAnywhere, Category=Settings)
+		FKBoxElem Box;
+	
+		// Produced geometry
+		UPROPERTY(meta = (DataflowOutput))
+		FSimpleGeometry Geometry;
+
+		void Evaluate(UE::Dataflow::FContext& Context, const FDataflowOutput* Out) const override;
+		void Register();
+	};
+
+	////////////////////////////////////////////////////////////////////////////////////////////////////////////////
+
+	USTRUCT()
+	struct FMakeSphereElemDataflowNode : public FDataflowNode
+	{
+		GENERATED_BODY()
+		DATAFLOW_NODE_DEFINE_INTERNAL(FMakeSphereElemDataflowNode, "MakeSphere", "Rigid Shapes", "")
+	
+	public:
+	
+		FMakeSphereElemDataflowNode(const UE::Dataflow::FNodeParameters& InParam, FGuid InGuid = FGuid::NewGuid())
+			: FDataflowNode(InParam, InGuid)
+		{
+			Register();
+		}
+
+	private:
+
+		// Input Sphere data
+		UPROPERTY(EditAnywhere, Category=Settings)
+		FKSphereElem Sphere;
+	
+		// Produced geometry
+		UPROPERTY(meta = (DataflowOutput))
+		FSimpleGeometry Geometry;
+
+		void Evaluate(UE::Dataflow::FContext& Context, const FDataflowOutput* Out) const override;
+		void Register();
+	};
+
+	////////////////////////////////////////////////////////////////////////////////////////////////////////////////
+
+	USTRUCT()
+	struct FMakeCapsuleElemDataflowNode : public FDataflowNode
+	{
+		GENERATED_BODY()
+		DATAFLOW_NODE_DEFINE_INTERNAL(FMakeCapsuleElemDataflowNode, "MakeCapsule", "Rigid Shapes", "")
+	
+	public:
+	
+		FMakeCapsuleElemDataflowNode(const UE::Dataflow::FNodeParameters& InParam, FGuid InGuid = FGuid::NewGuid())
+			: FDataflowNode(InParam, InGuid)
+		{
+			Register();
+		}
+
+	private:
+
+		// Input capsule data
+		UPROPERTY(EditAnywhere, Category=Settings)
+		FKSphylElem Capsule;
+	
+		// Produced geometry
+		UPROPERTY(meta = (DataflowOutput))
+		FSimpleGeometry Geometry;
+
+		void Evaluate(UE::Dataflow::FContext& Context, const FDataflowOutput* Out) const override;
+		void Register();
+	};
+
+	////////////////////////////////////////////////////////////////////////////////////////////////////////////////
+}
\ No newline at end of file
diff --git a/Engine/Plugins/Experimental/ChaosRigidAsset/Source/ChaosRigidAssetNodes/Public/BoneSelection.h b/Engine/Plugins/Experimental/ChaosRigidAsset/Source/ChaosRigidAssetNodes/Public/BoneSelection.h
new file mode 100644
--- /dev/null
+++ b/Engine/Plugins/Experimental/ChaosRigidAsset/Source/ChaosRigidAssetNodes/Public/BoneSelection.h
@@ -0,0 +1,67 @@
+// Copyright Epic Games, Inc. All Rights Reserved.
+
+#pragma once
+
+#include "UObject/ObjectPtr.h"
+#include "Containers/Array.h"
+
+#include "BoneSelection.generated.h"
+
+class USkeleton;
+class USkeletalMesh;
+
+/**
+ * Represents a single bone in a selection
+ */
+USTRUCT()
+struct FRigidAssetBoneInfo
+{
+	GENERATED_BODY()
+
+	FRigidAssetBoneInfo() = default;
+	FRigidAssetBoneInfo(FName InName, int32 InIndex, int32 InDepth);
+
+	// Bone name in the skeleton
+	UPROPERTY()
+	FName Name = NAME_None;
+
+	// Raw bone index in the skeleton
+	UPROPERTY()
+	int32 Index = INDEX_NONE;
+
+	// Depth of this bone in the hierarchy (distance from root)
+	UPROPERTY()
+	int32 Depth = INDEX_NONE;
+
+	bool operator ==(const FRigidAssetBoneInfo& Other) const = default;
+	bool operator <(const FRigidAssetBoneInfo& Other) const;
+};
+
+/**
+ * Represents a selection set of bones within a skeleton
+ */
+USTRUCT()
+struct FRigidAssetBoneSelection
+{
+	GENERATED_BODY()
+
+	// Individual bones in the selection
+	UPROPERTY()
+	TArray<FRigidAssetBoneInfo> SelectedBones;
+
+	// Skeleton that the bones reference
+	UPROPERTY()
+	TObjectPtr<USkeleton> Skeleton;
+
+	// Mesh that the bones reference
+	UPROPERTY()
+	TObjectPtr<USkeletalMesh> Mesh;
+
+	// Contains for index and name, cannot check the skeleton matches, so care should be taken
+	// that the search items here are representative of the mesh and skeleton bound to the selection
+	bool ContainsIndex(int32 InBoneIndex) const;
+	bool Contains(FName InBoneName) const;
+
+	// Sorts bones so parents are before children
+	void SortBones();
+};
diff --git a/Engine/Plugins/Experimental/ChaosRigidAsset/Source/ChaosRigidAssetNodes/Public/ChaosRigidAssetNodesModule.h b/Engine/Plugins/Experimental/ChaosRigidAsset/Source/ChaosRigidAssetNodes/Public/ChaosRigidAssetNodesModule.h
new file mode 100644
--- /dev/null
+++ b/Engine/Plugins/Experimental/ChaosRigidAsset/Source/ChaosRigidAssetNodes/Public/ChaosRigidAssetNodesModule.h
@@ -0,0 +1,15 @@
+// Copyright Epic Games, Inc. All Rights Reserved.
+
+#pragma once
+
+#include "Modules/ModuleInterface.h"
+
+class FChaosRigidAssetNodesModule : public IModuleInterface
+{
+public:
+
+	CHAOSRIGIDASSETNODES_API virtual void StartupModule() override;
+	CHAOSRIGIDASSETNODES_API virtual void ShutdownModule() override;
+
+private:
+};
\ No newline at end of file
diff --git a/Engine/Plugins/Experimental/ChaosRigidAsset/Source/ChaosRigidAssetNodes/Public/CopyOnWriteArray.h b/Engine/Plugins/Experimental/ChaosRigidAsset/Source/ChaosRigidAssetNodes/Public/CopyOnWriteArray.h
new file mode 100644
--- /dev/null
+++ b/Engine/Plugins/Experimental/ChaosRigidAsset/Source/ChaosRigidAssetNodes/Public/CopyOnWriteArray.h
@@ -0,0 +1,272 @@
+// Copyright Epic Games, Inc. All Rights Reserved.
+
+#pragma once
+
+#include "Containers/Array.h"
+#include "Concepts/ConvertibleTo.h"
+#include "Misc/ScopeRWLock.h"
+#include "UObject/UObjectGlobals.h"
+#include "UObject/ObjectPtr.h"
+#include "Dataflow/DataflowContextCache.h"
+
+class UObject;
+class UClass;
+class UScriptStruct;
+
+/**
+ * Concept that detects a ptr to UObject type, requiring:
+ *   - T is a ptr type
+ *   - T is convertible UObject*
+ *   - T Has a method GetClass the returns a UClass
+ */
+template<typename T>
+concept CIsUObjectType = requires (T Object)
+{
+	requires UE::CConvertibleTo<T, UObject*>;
+	{ Object->GetClass() } -> UE::CConvertibleTo<UClass*>;
+};
+
+/**
+ * Concept that detects a USTRUCT() type, requiring:
+ *   - T has a static method StaticStruct that returns a UScriptStruct*
+ */
+template<typename T>
+concept CIsUStructType = requires (T Object)
+{
+	{ T::StaticStruct() } -> UE::CConvertibleTo<UScriptStruct*>;
+};
+
+/**
+ * Wrapper for an array that implements copy-on-write for the array.
+ * If there are N readers but no writers there will be a single copy of the array
+ * If a codepath writes to the array and they aren't the only referencer, the array will be duplicated
+ * This means each modification will not implicitly duplicate - only if there are multiple referencers
+ * otherwise, we modify in place.
+ */
+template<typename T>
+class TCopyOnWriteArray
+{
+public:
+
+	TCopyOnWriteArray()
+	{
+
+	}
+
+	TCopyOnWriteArray(const TCopyOnWriteArray& Other)
+	{
+		State = Other.State;
+	}
+
+	TCopyOnWriteArray(TCopyOnWriteArray&& Other)
+	{
+		State = MoveTemp(Other.State);
+	}
+
+	TCopyOnWriteArray& operator=(const TCopyOnWriteArray& Other)
+	{
+		State = Other.State;
+		return *this;
+	}
+
+	TCopyOnWriteArray& operator=(TCopyOnWriteArray&& Other)
+	{
+		State = MoveTemp(Other.State);
+		return *this;
+	}
+
+	const T& operator [](int32 Index)
+	{
+		return State->Arr[Index];
+	}
+
+	TArray<T>::RangedForConstIteratorType begin() const
+	{
+		if(!State)
+		{
+			static const TArray<T> EmptyArray;
+			return EmptyArray.begin();
+		}
+
+		const TArray<T>& Arr = State->Arr;
+		return Arr.begin();
+	}
+
+	TArray<T>::RangedForConstIteratorType end() const
+	{
+		if(!State)
+		{
+			static const TArray<T> EmptyArray;
+			return EmptyArray.end();
+		}
+
+		const TArray<T>& Arr = State->Arr;
+		return Arr.end();
+	}
+
+	template<typename... Args>
+	int32 Emplace(Args&&... InArgs)
+	{
+		FWriteScopeLock ScopeLock(Lock);
+
+		CheckCopyOnWrite();
+		State->Arr.Emplace(Forward<Args>(InArgs)...);
+
+		return State->Arr.Num();
+	}
+
+	int32 Add(const T& Value)
+	{
+		FWriteScopeLock ScopeLock(Lock);
+
+		CheckCopyOnWrite();
+		State->Arr.Add(Value);
+
+		return State->Arr.Num();
+	}
+
+	int32 Add(T&& Value)
+	{
+		FWriteScopeLock ScopeLock(Lock);
+
+		CheckCopyOnWrite();
+		State->Arr.Add(Forward<T>(Value));
+
+		return State->Arr.Num();
+	}
+
+	void Modify(int32 AtIndex, const T& Value)
+	{
+		FWriteScopeLock ScopeLock(Lock);
+
+		CheckCopyOnWrite();
+		State->Arr[AtIndex] = Value;
+	}
+
+	void Modify(int32 AtIndex, T&& Value)
+	{
+		FWriteScopeLock ScopeLock(Lock);
+
+		CheckCopyOnWrite();
+		State->Arr[AtIndex] = MoveTemp(Value);
+	}
+
+	void Append(const TArray<T>& Other)
+	{
+		FWriteScopeLock ScopeLock(Lock);
+
+		CheckCopyOnWrite();
+		State->Arr.Append(Other);
+	}
+
+	// If the array contains UObjects, we need to emit references to them directly
+	void AddReferencedObjects(FReferenceCollector& Collector) const
+		requires CIsUObjectType<T>
+	{
+		if(!State)
+		{
+			return;
+		}
+
+		for(const T& Item : State->Arr)
+		{
+			TObjectPtr<UObject> AsObject = Item;
+			Collector.AddReferencedObject(AsObject);
+			AsObject->AddReferencedObjects(AsObject, Collector);
+		}
+	}
+
+	// If the array contains USTRUCT types, they might reference UObject types and
+	// need to emit references. Defer to the struct here to add required references
+	void AddReferencedObjects(FReferenceCollector& Collector) const
+		requires CIsUStructType<T>
+	{
+		if(!State)
+		{
+			return;
+		}
+
+		UScriptStruct* Struct = T::StaticStruct();
+
+		for(const T& Item : State->Arr)
+		{
+			Struct->GetCppStructOps()->AddStructReferencedObjects()((void*)&Item, Collector);
+		}
+	}
+
+	// Base case for storage types that don't emit references
+	void AddReferencedObjects(FReferenceCollector& Collector) const
+	{
+
+	}
+
+	// Copy the internal state into a new TArray
+	TArray<T> ToArray() const
+	{
+		if(!State)
+		{
+			return {};
+		}
+
+		TArray<T> NewArray;
+		NewArray.Reserve(State->Arr.Num());
+		NewArray.Append(State->Arr);
+
+		return MoveTemp(NewArray);
+	}
+
+private:
+
+	void CheckCopyOnWrite()
+	{
+		// If we have no state, or there are more than one reader, duplicate before writing
+		if(!State || State->GetRefCount() > 1)
+		{
+			Duplicate();
+		}
+	}
+
+	void Duplicate()
+	{
+		TRefCountPtr<FState> Curr = State;
+		TRefCountPtr<FState> New = new FState;
+
+		// Duplicate if there's a previous state
+		if(Curr)
+		{
+			// It's safe to iterate and append the old array here. If we're duplicating then all writers will duplicate.
+			// The current state is logically const at this point in time.
+			if constexpr(UE::Dataflow::template TIsUObjectPtrElement<T>::Value)
+			{
+				// UObjects need to run through StaticDuplicateObject rather than just be copied.
+				New->Arr.Reserve(Curr->Arr.Num());
+
+				for(const T& From : Curr->Arr)
+				{
+					New->Arr.Add(Cast<typename TRemoveObjectPointer<T>::Type>(StaticDuplicateObject(From, GetTransientPackage())));
+				}
+			}
+			else
+			{
+				New->Arr.Append(Curr->Arr);
+			}
+		}
+
+		State = New;
+	}
+
+	struct FState : public FThreadSafeRefCountedObject
+	{
+		TArray<T> Arr;
+	};
+
+	TRefCountPtr<FState> State = nullptr;
+	FRWLock Lock;
+};
+
+// Variadic helper to call AddReferencedObjects on multiple arrays
+template<typename... T>
+void AddArrayReferences(FReferenceCollector& Collector, const T&... Args)
+{
+	(Args.AddReferencedObjects(Collector), ...);
+}
\ No newline at end of file
diff --git a/Engine/Source/Runtime/Experimental/Chaos/Private/GeometryCollection/Facades/CollectionRenderingFacade.cpp b/Engine/Source/Runtime/Experimental/Chaos/Private/GeometryCollection/Facades/CollectionRenderingFacade.cpp
--- a/Engine/Source/Runtime/Experimental/Chaos/Private/GeometryCollection/Facades/CollectionRenderingFacade.cpp
+++ b/Engine/Source/Runtime/Experimental/Chaos/Private/GeometryCollection/Facades/CollectionRenderingFacade.cpp
@@ -660,6 +660,50 @@ namespace GeometryCollection::Facades
 		AddSurface(MoveTemp(NarrowVerts), MoveTemp(Triangles), MoveTemp(Generator.Normals), MoveTemp(Colours));
 	}
 
+	void FRenderingFacade::AddFaces(const TArray<FVector3f>& InVertices, TArray<FIntVector>& InIndices, TArray<FLinearColor>& InColors)
+	{
+		check(!IsConst());
+		if(IsValid())
+		{
+			TArray<FVector3f> ExpandedVerts;
+			TArray<FVector3f> ExpandedNormals;
+			TArray<FLinearColor> ExpandedColors;
+			TArray<FIntVector> ExpandedIndices;
+
+			const int32 NumExpandedVerts = InIndices.Num() * 3;
+			ExpandedVerts.SetNumUninitialized(NumExpandedVerts);
+			ExpandedNormals.SetNumUninitialized(NumExpandedVerts);
+			ExpandedColors.SetNumUninitialized(NumExpandedVerts);
+			ExpandedIndices.SetNumUninitialized(InIndices.Num());
+
+			int32 VertexBase = 0;
+			for(int32 FaceIndex = 0; FaceIndex < InIndices.Num(); ++FaceIndex)
+			{
+				const FIntVector Face = InIndices[FaceIndex];
+				const FVector3f AB = InVertices[Face[1]] - InVertices[Face[0]];
+				const FVector3f AC = InVertices[Face[2]] - InVertices[Face[0]];
+				const FVector3f FaceNormal = FVector3f::CrossProduct(AB, AC).GetSafeNormal();
+
+				ExpandedVerts[VertexBase + 0] = InVertices[Face[0]];
+				ExpandedVerts[VertexBase + 1] = InVertices[Face[1]];
+				ExpandedVerts[VertexBase + 2] = InVertices[Face[2]];
+
+				ExpandedNormals[VertexBase + 0] = FaceNormal;
+				ExpandedNormals[VertexBase + 1] = FaceNormal;
+				ExpandedNormals[VertexBase + 2] = FaceNormal;
+
+				ExpandedColors[VertexBase + 0] = InColors[Face[0]];
+				ExpandedColors[VertexBase + 1] = InColors[Face[1]];
+				ExpandedColors[VertexBase + 2] = InColors[Face[2]];
+
+				ExpandedIndices[FaceIndex] = {VertexBase, VertexBase + 1, VertexBase + 2};
+				VertexBase += 3;
+			}
+
+			AddSurface(MoveTemp(ExpandedVerts), MoveTemp(ExpandedIndices), MoveTemp(ExpandedNormals), MoveTemp(ExpandedColors));
+		}
+	}
+
 }; // GeometryCollection::Facades
 
 
diff --git a/Engine/Source/Runtime/Experimental/Chaos/Public/GeometryCollection/Facades/CollectionRenderingFacade.h b/Engine/Source/Runtime/Experimental/Chaos/Public/GeometryCollection/Facades/CollectionRenderingFacade.h
--- a/Engine/Source/Runtime/Experimental/Chaos/Public/GeometryCollection/Facades/CollectionRenderingFacade.h
+++ b/Engine/Source/Runtime/Experimental/Chaos/Public/GeometryCollection/Facades/CollectionRenderingFacade.h
@@ -58,6 +58,9 @@ namespace GeometryCollection::Facades
 		/**Add a triangle to the rendering view.*/
 		CHAOS_API void AddTriangle(const Chaos::FTriangle& InTriangle);
 
+		/** Given a list of triangles, expand them to faces with face-normals */
+		CHAOS_API void AddFaces(const TArray<FVector3f>& InVertices, TArray<FIntVector>& InIndices, TArray<FLinearColor>& InColors);
+
 		/**Add a box to the rendering view.*/
 		CHAOS_API void AddBox(const FBox& InBox);
 		CHAOS_API void AddBox(const FVector3f& InMinVertex, const FVector3f& InMaxVertex);
