Fixed registration and bounds checking in RigidAsset plugin

diff --git a/Engine/Plugins/Experimental/ChaosRigidAsset/Source/ChaosRigidAssetEditor/Private/ChaosRigidAssetEditorModule.cpp b/Engine/Plugins/Experimental/ChaosRigidAsset/Source/ChaosRigidAssetEditor/Private/ChaosRigidAssetEditorModule.cpp
--- a/Engine/Plugins/Experimental/ChaosRigidAsset/Source/ChaosRigidAssetEditor/Private/ChaosRigidAssetEditorModule.cpp
+++ b/Engine/Plugins/Experimental/ChaosRigidAsset/Source/ChaosRigidAssetEditor/Private/ChaosRigidAssetEditorModule.cpp
@@ -186,7 +186,9 @@ TSharedRef<FExtender> ExtendAssetContextMenu(const TArray<FAssetData>& Assets)
 void FChaosRigidAssetEditorModule::StartupModule()
 {
 	FContentBrowserModule& ContentBrowser = FModuleManager::LoadModuleChecked<FContentBrowserModule>(TEXT("ContentBrowser"));
-	ContentBrowser.GetAllAssetViewContextMenuExtenders().Add(FContentBrowserMenuExtender_SelectedAssets::CreateStatic(&ExtendAssetContextMenu));
+	TArray<FContentBrowserMenuExtender_SelectedAssets>& AssetMenuExtenders = ContentBrowser.GetAllAssetViewContextMenuExtenders();
+	AssetMenuExtenders.Add(FContentBrowserMenuExtender_SelectedAssets::CreateStatic(&ExtendAssetContextMenu));
+	AssetMenuExtenderHandle = AssetMenuExtenders.Last().GetHandle();
 
 	EditorFeature = MakeUnique<FDataflowPhysicsAssetEditorOverride>();
 
@@ -201,6 +203,23 @@ void FChaosRigidAssetEditorModule::ShutdownModule()
 {
 	IModularFeatures::Get().UnregisterModularFeature(IPhysicsAssetEditorOverride::ModularFeatureName, EditorFeature.Get());
 	EditorFeature = nullptr;
+
+	if(AssetMenuExtenderHandle.IsValid() && FModuleManager::Get().IsModuleLoaded(TEXT("ContentBrowser")))
+	{
+		FContentBrowserModule& ContentBrowserModule = FModuleManager::LoadModuleChecked<FContentBrowserModule>(TEXT("ContentBrowser"));
+		TArray<FContentBrowserMenuExtender_SelectedAssets>& AssetMenuExtenders = ContentBrowserModule.GetAllAssetViewContextMenuExtenders();
+
+		AssetMenuExtenders.RemoveAll([this](const FContentBrowserMenuExtender_SelectedAssets& InDelegate)
+			{
+				return InDelegate.GetHandle() == AssetMenuExtenderHandle;
+			});
+
+		AssetMenuExtenderHandle.Reset();
+	}
+
+	UE::Dataflow::FRenderingFactory::GetInstance()->DeregisterCallbacks(UE::Chaos::RigidAsset::FPhysAssetStateRenderCallbacks::StaticGetRenderKey());
+	UE::Dataflow::FRenderingFactory::GetInstance()->DeregisterCallbacks(UE::Chaos::RigidAsset::FBoneSelectionRenderCallbacks::StaticGetRenderKey());
+	UE::Dataflow::FRenderingFactory::GetInstance()->DeregisterCallbacks(UE::Chaos::RigidAsset::FAggregateGeometryGeomRenderCallbacks::StaticGetRenderKey());
 }
 
 bool FDataflowPhysicsAssetEditorOverride::OpenAsset(UPhysicsAsset* InAsset)
diff --git a/Engine/Plugins/Experimental/ChaosRigidAsset/Source/ChaosRigidAssetEditor/Private/DataflowRendering.cpp b/Engine/Plugins/Experimental/ChaosRigidAsset/Source/ChaosRigidAssetEditor/Private/DataflowRendering.cpp
--- a/Engine/Plugins/Experimental/ChaosRigidAsset/Source/ChaosRigidAssetEditor/Private/DataflowRendering.cpp
+++ b/Engine/Plugins/Experimental/ChaosRigidAsset/Source/ChaosRigidAssetEditor/Private/DataflowRendering.cpp
@@ -78,11 +78,16 @@ namespace UE::Chaos::RigidAsset
 		}
 	}
 
-	UE::Dataflow::FRenderKey FAggregateGeometryGeomRenderCallbacks::GetRenderKey() const
+	UE::Dataflow::FRenderKey FAggregateGeometryGeomRenderCallbacks::StaticGetRenderKey()
 	{
 		return { "GeomRender", FName("FKAggregateGeom") };
 	}
 
+	UE::Dataflow::FRenderKey FAggregateGeometryGeomRenderCallbacks::GetRenderKey() const
+	{
+		return StaticGetRenderKey();
+	}
+
 	bool FAggregateGeometryGeomRenderCallbacks::CanRender(const UE::Dataflow::IDataflowConstructionViewMode& ViewMode) const
 	{
 		return true;
@@ -100,11 +105,16 @@ namespace UE::Chaos::RigidAsset
 		}
 	}
 
-	UE::Dataflow::FRenderKey FBoneSelectionRenderCallbacks::GetRenderKey() const
+	UE::Dataflow::FRenderKey FBoneSelectionRenderCallbacks::StaticGetRenderKey()
 	{
 		return { "GeomRender", FName("FBoneSelection") };
 	}
 
+	UE::Dataflow::FRenderKey FBoneSelectionRenderCallbacks::GetRenderKey() const
+	{
+		return StaticGetRenderKey();
+	}
+
 	bool FBoneSelectionRenderCallbacks::CanRender(const UE::Dataflow::IDataflowConstructionViewMode& ViewMode) const
 	{
 		return true;
@@ -238,11 +248,16 @@ namespace UE::Chaos::RigidAsset
 		}
 	}
 
-	UE::Dataflow::FRenderKey FPhysAssetStateRenderCallbacks::GetRenderKey() const
+	UE::Dataflow::FRenderKey FPhysAssetStateRenderCallbacks::StaticGetRenderKey()
 	{
 		return { "GeomRender", FName("FPhysicsAssetDataflowState") };
 	}
 
+	UE::Dataflow::FRenderKey FPhysAssetStateRenderCallbacks::GetRenderKey() const
+	{
+		return StaticGetRenderKey();
+	}
+
 	bool FPhysAssetStateRenderCallbacks::CanRender(const UE::Dataflow::IDataflowConstructionViewMode& ViewMode) const
 	{
 		return true;
@@ -262,9 +277,12 @@ namespace UE::Chaos::RigidAsset
 				for(TObjectPtr<USkeletalBodySetup> Body : AssetState.Bodies)
 				{
 					const int32 BoneIndex = RefSkeleton.FindBoneIndex(Body->BoneName);
-					FTransform BoneTransform = RefSkeleton.GetBoneAbsoluteTransform(BoneIndex);
 
-					RenderAggGeom(RenderingFacade, State, Body->AggGeom, BoneTransform, Body->BoneName, FLinearColor(1.0f, 1.0f, 1.0f, 0.5f));
+					if(BoneIndex != INDEX_NONE)
+					{
+						FTransform BoneTransform = RefSkeleton.GetBoneAbsoluteTransform(BoneIndex);
+						RenderAggGeom(RenderingFacade, State, Body->AggGeom, BoneTransform, Body->BoneName, FLinearColor(1.0f, 1.0f, 1.0f, 0.5f));
+					}
 				}
 			}
 		}
diff --git a/Engine/Plugins/Experimental/ChaosRigidAsset/Source/ChaosRigidAssetEditor/Public/ChaosRigidAssetEditorModule.h b/Engine/Plugins/Experimental/ChaosRigidAsset/Source/ChaosRigidAssetEditor/Public/ChaosRigidAssetEditorModule.h
--- a/Engine/Plugins/Experimental/ChaosRigidAsset/Source/ChaosRigidAssetEditor/Public/ChaosRigidAssetEditorModule.h
+++ b/Engine/Plugins/Experimental/ChaosRigidAsset/Source/ChaosRigidAssetEditor/Public/ChaosRigidAssetEditorModule.h
@@ -4,6 +4,7 @@
 
 #include "Modules/ModuleInterface.h"
 #include "IPhysicsAssetEditor.h"
+#include "Delegates/IDelegateInstance.h"
 
 class FDataflowPhysicsAssetEditorOverride;
 
@@ -16,4 +17,5 @@ public:
 
 private:
 	TUniquePtr<FDataflowPhysicsAssetEditorOverride> EditorFeature;
+	FDelegateHandle AssetMenuExtenderHandle;
 };
diff --git a/Engine/Plugins/Experimental/ChaosRigidAsset/Source/ChaosRigidAssetEditor/Public/DataflowRendering.h b/Engine/Plugins/Experimental/ChaosRigidAsset/Source/ChaosRigidAssetEditor/Public/DataflowRendering.h
--- a/Engine/Plugins/Experimental/ChaosRigidAsset/Source/ChaosRigidAssetEditor/Public/DataflowRendering.h
+++ b/Engine/Plugins/Experimental/ChaosRigidAsset/Source/ChaosRigidAssetEditor/Public/DataflowRendering.h
@@ -10,6 +10,8 @@ namespace UE::Chaos::RigidAsset
 	{
 	public:
 
+		static UE::Dataflow::FRenderKey StaticGetRenderKey();
+
 		UE::Dataflow::FRenderKey GetRenderKey() const override;
 		bool CanRender(const UE::Dataflow::IDataflowConstructionViewMode& ViewMode) const override;
 		void Render(GeometryCollection::Facades::FRenderingFacade& RenderingFacade, const UE::Dataflow::FGraphRenderingState& State) override;
@@ -20,6 +22,8 @@ namespace UE::Chaos::RigidAsset
 	{
 	public:
 
+		static UE::Dataflow::FRenderKey StaticGetRenderKey();
+
 		UE::Dataflow::FRenderKey GetRenderKey() const override;
 		bool CanRender(const UE::Dataflow::IDataflowConstructionViewMode& ViewMode) const override;
 		void Render(GeometryCollection::Facades::FRenderingFacade& RenderingFacade, const UE::Dataflow::FGraphRenderingState& State) override;
@@ -29,6 +33,8 @@ namespace UE::Chaos::RigidAsset
 	{
 	public:
 
+		static UE::Dataflow::FRenderKey StaticGetRenderKey();
+
 		UE::Dataflow::FRenderKey GetRenderKey() const override;
 		bool CanRender(const UE::Dataflow::IDataflowConstructionViewMode& ViewMode) const override;
 		void Render(GeometryCollection::Facades::FRenderingFacade& RenderingFacade, const UE::Dataflow::FGraphRenderingState& State) override;
diff --git a/Engine/Plugins/Experimental/ChaosRigidAsset/Source/ChaosRigidAssetNodes/Private/Generators/BoneGeometryGenerators.cpp b/Engine/Plugins/Experimental/ChaosRigidAsset/Source/ChaosRigidAssetNodes/Private/Generators/BoneGeometryGenerators.cpp
--- a/Engine/Plugins/Experimental/ChaosRigidAsset/Source/ChaosRigidAssetNodes/Private/Generators/BoneGeometryGenerators.cpp
+++ b/Engine/Plugins/Experimental/ChaosRigidAsset/Source/ChaosRigidAssetNodes/Private/Generators/BoneGeometryGenerators.cpp
@@ -11,6 +11,7 @@
 #include "Rendering/SkinWeightVertexBuffer.h"
 #include "Rendering/PositionVertexBuffer.h"
 #include "Rendering/StaticMeshVertexBuffer.h"
+#include "Math/Box.h"
 
 namespace UE::Dataflow
 {
@@ -239,30 +240,10 @@ TArray<UE::Chaos::RigidAsset::FSimpleGeometry> UBoneGeometryGenerator_Sphere::Bu
 			continue;
 		}
 
-		FVector2f AxisIntervals[] =
-		{
-			{std::numeric_limits<float>::max(), std::numeric_limits<float>::lowest()},
-			{std::numeric_limits<float>::max(), std::numeric_limits<float>::lowest()},
-			{std::numeric_limits<float>::max(), std::numeric_limits<float>::lowest()}
-		};
-
-		for(const FVector3f& Position : BoneVerts.Positions)
-		{
-			AxisIntervals[0] = { FMath::Min(AxisIntervals[0][0], Position[0]), FMath::Max(AxisIntervals[0][1], Position[0]) };
-			AxisIntervals[1] = { FMath::Min(AxisIntervals[1][0], Position[1]), FMath::Max(AxisIntervals[1][1], Position[1]) };
-			AxisIntervals[2] = { FMath::Min(AxisIntervals[2][0], Position[2]), FMath::Max(AxisIntervals[2][1], Position[2]) };
-		}
-
-		const FVector3f SphereCentre(
-			(AxisIntervals[0][0] + AxisIntervals[0][1]) * 0.5f,
-			(AxisIntervals[1][0] + AxisIntervals[1][1]) * 0.5f,
-			(AxisIntervals[2][0] + AxisIntervals[2][1]) * 0.5f
-		);
-
-		const float SphereRadius = FMath::Max(AxisIntervals[0][1] - AxisIntervals[0][0], AxisIntervals[1][1] - AxisIntervals[1][0], AxisIntervals[2][1] - AxisIntervals[2][0]);
+		FBox3f VertBox(BoneVerts.Positions);
 
 		Result.AddDefaulted();
-		Result.Last().GeometryElement = MakeShared<FKSphereElem>(SphereRadius);
+		Result.Last().GeometryElement = MakeShared<FKSphereElem>(VertBox.GetExtent().Length());
 	}
 
 	return Result;
diff --git a/Engine/Plugins/Experimental/ChaosRigidAsset/Source/ChaosRigidAssetNodes/Private/PhysicsAssetBuilder.cpp b/Engine/Plugins/Experimental/ChaosRigidAsset/Source/ChaosRigidAssetNodes/Private/PhysicsAssetBuilder.cpp
--- a/Engine/Plugins/Experimental/ChaosRigidAsset/Source/ChaosRigidAssetNodes/Private/PhysicsAssetBuilder.cpp
+++ b/Engine/Plugins/Experimental/ChaosRigidAsset/Source/ChaosRigidAssetNodes/Private/PhysicsAssetBuilder.cpp
@@ -94,8 +94,13 @@ namespace UE::Chaos::RigidAsset
 		if(!TargetAsset)
 		{
 			CreateNewAsset();
+
+			// If we fail to make an asset, abort
+			if(!TargetAsset)
+			{
+				return nullptr;
+			}
 		}
-		check(TargetAsset);
 
 		// Need to replace the objects already in the asset, first move them to the transient package to get them out of the physics asset package
 		for(TObjectPtr<USkeletalBodySetup> BodyToRename : TargetAsset->SkeletalBodySetups)
diff --git a/Engine/Plugins/Experimental/ChaosRigidAsset/Source/ChaosRigidAssetNodes/Private/PhysicsAssetDataflowNodes.cpp b/Engine/Plugins/Experimental/ChaosRigidAsset/Source/ChaosRigidAssetNodes/Private/PhysicsAssetDataflowNodes.cpp
--- a/Engine/Plugins/Experimental/ChaosRigidAsset/Source/ChaosRigidAssetNodes/Private/PhysicsAssetDataflowNodes.cpp
+++ b/Engine/Plugins/Experimental/ChaosRigidAsset/Source/ChaosRigidAssetNodes/Private/PhysicsAssetDataflowNodes.cpp
@@ -696,7 +696,7 @@ void FDataflowNewBoneSelection::Evaluate(UE::Dataflow::FContext& Context, const
 		{
 			if(TObjectPtr<UPhysicsAsset> PhysAsset = GetAttachmentOwnerAs<UPhysicsAsset>(Context))
 			{
-				InSkeleton = PhysAsset->GetPreviewMesh()->GetSkeleton();
+				InSkeleton = PhysAsset->GetPreviewMesh() ? PhysAsset->GetPreviewMesh()->GetSkeleton() : nullptr;
 			}
 			else
 			{
@@ -761,6 +761,13 @@ void FDataflowPhysicsAssetMakeState::Evaluate(UE::Dataflow::FContext& Context, c
 			if(TObjectPtr<UPhysicsAsset> PhysAsset = GetAttachmentOwnerAs<UPhysicsAsset>(Context))
 			{
 				InTarget = PhysAsset->GetPreviewMesh();
+
+				if(!InTarget)
+				{
+					Context.Error(TEXT("Attempted to make a state without a valid mesh"), this);
+					SetValue(Context, FPhysicsAssetDataflowState{}, &State);
+					return;
+				}
 			}
 			else
 			{
diff --git a/Engine/Plugins/Experimental/ChaosRigidAsset/Source/ChaosRigidAssetNodes/Private/PhysicsAssetDataflowNodes.h b/Engine/Plugins/Experimental/ChaosRigidAsset/Source/ChaosRigidAssetNodes/Private/PhysicsAssetDataflowNodes.h
--- a/Engine/Plugins/Experimental/ChaosRigidAsset/Source/ChaosRigidAssetNodes/Private/PhysicsAssetDataflowNodes.h
+++ b/Engine/Plugins/Experimental/ChaosRigidAsset/Source/ChaosRigidAssetNodes/Private/PhysicsAssetDataflowNodes.h
@@ -307,7 +307,7 @@ private:
 	FRigidAssetBoneSelection B;
 
 	// Resulting selection, a combination of A and B
-	UPROPERTY(meta = (DatflowOutput))
+	UPROPERTY(meta = (DataflowOutput))
 	FRigidAssetBoneSelection Selection;
 
 	void Evaluate(UE::Dataflow::FContext& Context, const FDataflowOutput* Out) const override;
diff --git a/Engine/Plugins/Experimental/ChaosRigidAsset/Source/ChaosRigidAssetNodes/Private/PhysicsAssetDataflowState.cpp b/Engine/Plugins/Experimental/ChaosRigidAsset/Source/ChaosRigidAssetNodes/Private/PhysicsAssetDataflowState.cpp
--- a/Engine/Plugins/Experimental/ChaosRigidAsset/Source/ChaosRigidAssetNodes/Private/PhysicsAssetDataflowState.cpp
+++ b/Engine/Plugins/Experimental/ChaosRigidAsset/Source/ChaosRigidAssetNodes/Private/PhysicsAssetDataflowState.cpp
@@ -7,6 +7,7 @@
 void FPhysicsAssetDataflowState::AddStructReferencedObjects(FReferenceCollector& Collector)
 {
 	Collector.AddReferencedObject(TargetSkeleton);
+	Collector.AddReferencedObject(TargetMesh);
 	AddArrayReferences(Collector, Bodies, Constraints);
 }
 
diff --git a/Engine/Plugins/Experimental/ChaosRigidAsset/Source/ChaosRigidAssetNodes/Private/RigidDataflowNode.h b/Engine/Plugins/Experimental/ChaosRigidAsset/Source/ChaosRigidAssetNodes/Private/RigidDataflowNode.h
--- a/Engine/Plugins/Experimental/ChaosRigidAsset/Source/ChaosRigidAssetNodes/Private/RigidDataflowNode.h
+++ b/Engine/Plugins/Experimental/ChaosRigidAsset/Source/ChaosRigidAssetNodes/Private/RigidDataflowNode.h
@@ -173,6 +173,11 @@ struct TMultiPinArrayConnector : public IMultiPinArrayConnector
 
 	TArray<UE::Dataflow::FPin> GetPinsToRemove() const override
 	{
+		if(PinContainer.Num() == 0)
+		{
+			return {};
+		}
+
 		const int32 Index = PinContainer.Num() - 1;
 		UE::Dataflow::TConnectionReference<UE::Chaos::RigidAsset::FSimpleGeometry> Ref{ &PinContainer[Index], Index, &PinContainer };
 
diff --git a/Engine/Plugins/Experimental/ChaosRigidAsset/Source/ChaosRigidAssetNodes/Public/CopyOnWriteArray.h b/Engine/Plugins/Experimental/ChaosRigidAsset/Source/ChaosRigidAssetNodes/Public/CopyOnWriteArray.h
--- a/Engine/Plugins/Experimental/ChaosRigidAsset/Source/ChaosRigidAssetNodes/Public/CopyOnWriteArray.h
+++ b/Engine/Plugins/Experimental/ChaosRigidAsset/Source/ChaosRigidAssetNodes/Public/CopyOnWriteArray.h
@@ -77,6 +77,7 @@ public:
 
 	const T& operator [](int32 Index)
 	{
+		checkf(State, TEXT("Array index out of bounds: %d into an array of size 0"), Index);
 		return State->Arr[Index];
 	}
 
