[Dataflow] Added Box, Capsule and Convex simple builders as geometry generators for dataflow physics assets.

diff --git a/Engine/Plugins/Experimental/ChaosRigidAsset/Source/ChaosRigidAssetEditor/Private/DataflowRendering.cpp b/Engine/Plugins/Experimental/ChaosRigidAsset/Source/ChaosRigidAssetEditor/Private/DataflowRendering.cpp
--- a/Engine/Plugins/Experimental/ChaosRigidAsset/Source/ChaosRigidAssetEditor/Private/DataflowRendering.cpp
+++ b/Engine/Plugins/Experimental/ChaosRigidAsset/Source/ChaosRigidAssetEditor/Private/DataflowRendering.cpp
@@ -43,7 +43,7 @@ namespace UE::Chaos::RigidAsset
 				const FKBoxElem* AsBox = static_cast<const FKBoxElem*>(Elem);
 				const FVector3f Extent = FVector3f{ AsBox->X, AsBox->Y, AsBox->Z } *0.5f;
 
-				RenderingFacade.SetGroupTransform(GeometryGroup, BaseTm * Elem->GetTransform());
+				RenderingFacade.SetGroupTransform(GeometryGroup, Elem->GetTransform() * BaseTm);
 				RenderingFacade.AddBox({ -Extent, Extent });
 
 				break;
@@ -52,8 +52,9 @@ namespace UE::Chaos::RigidAsset
 			{
 				const FKSphereElem* AsSphere = static_cast<const FKSphereElem*>(Elem);
 
-				RenderingFacade.SetGroupTransform(GeometryGroup, BaseTm * Elem->GetTransform());
-				RenderingFacade.AddSphere((FVector3f)AsSphere->Center, AsSphere->Radius, InColor);
+				// Transform placed in the group - always draw at (0,0,0)
+				RenderingFacade.SetGroupTransform(GeometryGroup, Elem->GetTransform() * BaseTm);
+				RenderingFacade.AddSphere({0.0f, 0.0f, 0.0f}, AsSphere->Radius, InColor);
 
 				break;
 			}
@@ -65,11 +66,55 @@ namespace UE::Chaos::RigidAsset
 				// the segment, we need to transform first to the centre of the capsule, then apply the elem transform
 				const FTransform CapsuleCollisionAdjust(FQuat::Identity, FVector(0.0f, 0.0f, -AsSphyl->Length / 2));
 
-				RenderingFacade.SetGroupTransform(GeometryGroup, BaseTm * CapsuleCollisionAdjust * Elem->GetTransform());
+				RenderingFacade.SetGroupTransform(GeometryGroup, CapsuleCollisionAdjust * Elem->GetTransform() * BaseTm);
 				RenderingFacade.AddCapsule(AsSphyl->Length, AsSphyl->Radius);
 
 				break;
 			}
+			case EAggCollisionShape::Convex:
+			{
+				const FKConvexElem* AsConvex = static_cast<const FKConvexElem*>(Elem);
+
+				const int32 NumIndices = AsConvex->IndexData.Num();
+				if(NumIndices == 0 || NumIndices % 3 != 0)
+				{
+					// Empty or invalid convex index data.
+					break;
+				}
+
+				RenderingFacade.SetGroupTransform(GeometryGroup, Elem->GetTransform() * BaseTm);
+
+				const int32 NumFaces = NumIndices / 3;
+				TArray<FVector3f> Positions;
+				Positions.Reserve(AsConvex->VertexData.Num());
+				for(const FVector& Position : AsConvex->VertexData)
+				{
+					Positions.Add(static_cast<FVector3f>(Position));
+				}
+
+				TArray<FIntVector> Indices;
+				Indices.Reserve(NumFaces);
+
+				for(int32 FaceIndex = 0; FaceIndex < AsConvex->IndexData.Num(); FaceIndex += 3)
+				{
+					Indices.Add({
+						AsConvex->IndexData[FaceIndex + 0],
+						AsConvex->IndexData[FaceIndex + 2],
+						AsConvex->IndexData[FaceIndex + 1]
+					});
+				}
+
+				TArray<FLinearColor> Colors;
+				Colors.SetNum(AsConvex->VertexData.Num());
+				for(FLinearColor& Color : Colors)
+				{
+					Color = FLinearColor::White;
+				}
+
+				RenderingFacade.AddFaces(Positions, Indices, Colors);
+
+				break;
+			}
 			default:
 				break;
 			}
diff --git a/Engine/Plugins/Experimental/ChaosRigidAsset/Source/ChaosRigidAssetNodes/ChaosRigidAssetNodes.Build.cs b/Engine/Plugins/Experimental/ChaosRigidAsset/Source/ChaosRigidAssetNodes/ChaosRigidAssetNodes.Build.cs
--- a/Engine/Plugins/Experimental/ChaosRigidAsset/Source/ChaosRigidAssetNodes/ChaosRigidAssetNodes.Build.cs
+++ b/Engine/Plugins/Experimental/ChaosRigidAsset/Source/ChaosRigidAssetNodes/ChaosRigidAssetNodes.Build.cs
@@ -24,9 +24,11 @@ public class ChaosRigidAssetNodes : ModuleRules
 			new string[]
 			{
 				"CoreUObject",
+				"PhysicsUtilities"
 			}
 		);
 
+
 		SetupModulePhysicsSupport(Target);
 
 		bAllowUETypesInNamespaces = true;
diff --git a/Engine/Plugins/Experimental/ChaosRigidAsset/Source/ChaosRigidAssetNodes/Private/Generators/BoneGeometryGenerators.cpp b/Engine/Plugins/Experimental/ChaosRigidAsset/Source/ChaosRigidAssetNodes/Private/Generators/BoneGeometryGenerators.cpp
--- a/Engine/Plugins/Experimental/ChaosRigidAsset/Source/ChaosRigidAssetNodes/Private/Generators/BoneGeometryGenerators.cpp
+++ b/Engine/Plugins/Experimental/ChaosRigidAsset/Source/ChaosRigidAssetNodes/Private/Generators/BoneGeometryGenerators.cpp
@@ -1,17 +1,20 @@
 // Copyright Epic Games, Inc. All Rights Reserved.
 
-#include "Generators/BoneGeometryGenerators.h"
+#include "Dataflow/DataflowDebugDrawInterface.h"
 #include "Dataflow/DataflowNode.h"
 #include "Dataflow/DataflowNodeFactory.h"
-#include "Rendering/SkeletalMeshModel.h"
 #include "Engine/SkeletalMesh.h"
+#include "Generators/BoneGeometryGenerators.h"
+#include "Math/Box.h"
+#include "PhysicsAssetUtils.h"
+#include "PhysicsEngine/ConvexElem.h"
 #include "ReferenceSkeleton.h"
-#include "Rendering/SkeletalMeshRenderData.h"
+#include "Rendering/PositionVertexBuffer.h"
 #include "Rendering/SkeletalMeshLODRenderData.h"
+#include "Rendering/SkeletalMeshModel.h"
+#include "Rendering/SkeletalMeshRenderData.h"
 #include "Rendering/SkinWeightVertexBuffer.h"
-#include "Rendering/PositionVertexBuffer.h"
 #include "Rendering/StaticMeshVertexBuffer.h"
-#include "Math/Box.h"
 
 namespace UE::Dataflow
 {
@@ -99,6 +102,24 @@ void FMakeConvexBoneGeometryGenerator::Register()
 
 ////////////////////////////////////////////////////////////////////////////////////////////////////////////////
 
+#if WITH_EDITOR
+bool UBoneGeometryGenerator::CanDebugDraw() const
+{
+	return false;
+}
+
+bool UBoneGeometryGenerator::CanDebugDrawViewMode(const FName& ViewModeName) const
+{
+	return false;
+}
+
+void UBoneGeometryGenerator::DebugDraw(UE::Dataflow::FContext& Context, IDataflowDebugDrawInterface& DataflowRenderingInterface, const FDataflowNode::FDebugDrawParameters& DebugDrawParameters, FRigidAssetBoneSelection Bones) const
+{
+}
+#endif
+
+////////////////////////////////////////////////////////////////////////////////////////////////////////////////
+
 /**
  * Data extracted from an LOD for body generation
  */
@@ -213,7 +234,39 @@ FBoneVertData GetBoneVertexData(FName BoneName, USkeletalMesh* Mesh, int32 Lod,
 
 TArray<UE::Chaos::RigidAsset::FSimpleGeometry> UBoneGeometryGenerator_Box::Build(FRigidAssetBoneSelection Bones)
 {
-	return {};
+	if(!Bones.Mesh || !Bones.Skeleton)
+	{
+		return {};
+	}
+
+	TArray<UE::Chaos::RigidAsset::FSimpleGeometry> Result;
+
+	for(const FRigidAssetBoneInfo& Bone : Bones.SelectedBones)
+	{
+		FBoneVertData BoneVerts = GetBoneVertexData(Bone.Name, Bones.Mesh, BaseSettings.SourceLod, BaseSettings.VertexMode);
+
+		if(BoneVerts.Positions.Num() == 0)
+		{
+			// No influences, make a small sphere and warn
+			// #TODO pass context through for dataflow warnings
+			Result.AddDefaulted();
+			Result.Last().GeometryElement = MakeShared<FKSphereElem>(5.0f);
+			continue;
+		}
+
+		FBox3f VertBox(BoneVerts.Positions);
+
+		TSharedPtr<FKBoxElem> NewBox = MakeShared<FKBoxElem>();
+		NewBox->X = VertBox.GetExtent().X * 2.0f;
+		NewBox->Y = VertBox.GetExtent().Y * 2.0f;
+		NewBox->Z = VertBox.GetExtent().Z * 2.0f;
+		NewBox->Center = static_cast<FVector>(VertBox.GetCenter());
+
+		Result.AddDefaulted();
+		Result.Last().GeometryElement = NewBox;
+	}
+
+	return Result;
 }
 
 ////////////////////////////////////////////////////////////////////////////////////////////////////////////////
@@ -242,23 +295,270 @@ TArray<UE::Chaos::RigidAsset::FSimpleGeometry> UBoneGeometryGenerator_Sphere::Bu
 
 		FBox3f VertBox(BoneVerts.Positions);
 
+		TSharedPtr<FKSphereElem> NewSphere = MakeShared<FKSphereElem>();
+		NewSphere->Radius = VertBox.GetExtent().Length();
+		NewSphere->Center = static_cast<FVector>(VertBox.GetCenter());
+
 		Result.AddDefaulted();
-		Result.Last().GeometryElement = MakeShared<FKSphereElem>(VertBox.GetExtent().Length());
+		Result.Last().GeometryElement = NewSphere;
 	}
 
 	return Result;
 }
 
+#if WITH_EDITOR
+bool UBoneGeometryGenerator_Sphere::CanDebugDraw() const
+{
+	return true;
+}
+
+bool UBoneGeometryGenerator_Sphere::CanDebugDrawViewMode(const FName& ViewModeName) const
+{
+	return true;
+}
+
+void UBoneGeometryGenerator_Sphere::DebugDraw(UE::Dataflow::FContext& Context, IDataflowDebugDrawInterface& DataflowRenderingInterface, const FDataflowNode::FDebugDrawParameters& DebugDrawParameters, FRigidAssetBoneSelection Bones) const
+{
+	if(!Bones.Mesh || !Bones.Skeleton || !bDrawVerts)
+	{
+		return;
+	}
+
+	TArray<FTransform> BoneTransforms;
+	const FReferenceSkeleton& RefSkel = Bones.Mesh->GetRefSkeleton();
+	RefSkel.GetBoneAbsoluteTransforms(BoneTransforms);
+
+	int32 BoneColorSeed = 0;
+	for(const FRigidAssetBoneInfo& Bone : Bones.SelectedBones)
+	{
+		if(DrawVertsForBoneName != NAME_None && Bone.Name == DrawVertsForBoneName)
+		{
+			continue;
+		}
+
+		FBoneVertData BoneVerts = GetBoneVertexData(Bone.Name, Bones.Mesh, BaseSettings.SourceLod, BaseSettings.VertexMode);
+
+		DataflowRenderingInterface.SetColor(FLinearColor::MakeRandomSeededColor(BoneColorSeed++));
+		for(const FVector3f& Pos : BoneVerts.Positions)
+		{
+			DataflowRenderingInterface.DrawPoint(BoneTransforms[Bone.Index].TransformPosition(static_cast<FVector>(Pos)));
+		}
+	}
+}
+#endif
+
 ////////////////////////////////////////////////////////////////////////////////////////////////////////////////
 
 TArray<UE::Chaos::RigidAsset::FSimpleGeometry> UBoneGeometryGenerator_Capsule::Build(FRigidAssetBoneSelection Bones)
 {
-	return {};
+	if(!Bones.Mesh || !Bones.Skeleton)
+	{
+		return {};
+	}
+
+	TArray<UE::Chaos::RigidAsset::FSimpleGeometry> Result;
+
+	for(const FRigidAssetBoneInfo& Bone : Bones.SelectedBones)
+	{
+		FBoneVertData BoneVerts = GetBoneVertexData(Bone.Name, Bones.Mesh, BaseSettings.SourceLod, BaseSettings.VertexMode);
+
+		if(BoneVerts.Positions.Num() == 0)
+		{
+			// No influences, make a small sphere and warn
+			// #TODO pass context through for dataflow warnings
+			Result.AddDefaulted();
+			Result.Last().GeometryElement = MakeShared<FKSphereElem>(5.0f);
+			continue;
+		}
+
+		// Y and Z are flipped in the default case deliberately to line up a capsule
+		// correctly with the bone extent, otherwise it appears on its side
+		FVector3f XAxis(1, 0, 0);
+		FVector3f YAxis(0, 0, 1);
+		FVector3f ZAxis(0, 1, 0);
+
+		if(Alignment == EBodyAlignment::Verts)
+		{
+			FMatrix CovarianceMat = FPhysicsAssetUtils::ComputeViewCovarianceMatrix(MakeConstArrayView(BoneVerts.Positions));
+			ZAxis = static_cast<FVector3f>(FPhysicsAssetUtils::ComputeEigenVector(CovarianceMat));
+			ZAxis.FindBestAxisVectors(XAxis, YAxis);
+		}
+
+		auto AveragePosition = [](TArrayView<const FVector3f> Verts)
+			{
+				FVector3f Avg(0);
+
+				if(Verts.Num() == 0)
+				{
+					return Avg;
+				}
+
+				for(const FVector3f Vert : Verts)
+				{
+					Avg += Vert;
+				}
+
+				return Avg / static_cast<float>(Verts.Num());
+			};
+
+		auto AxisInterval = [](TArrayView<const FVector3f> Verts, FVector3f Axis)
+			{
+				float Min = std::numeric_limits<float>::max();
+				float Max = std::numeric_limits<float>::lowest();
+
+				for(const FVector3f Vert : Verts)
+				{
+					const float Projected = FVector3f::DotProduct(Vert, Axis);
+					if(Projected < Min)
+					{
+						Min = Projected;
+					}
+
+					if(Projected > Max)
+					{
+						Max = Projected;
+					}
+				}
+
+				return Max - Min;
+			};
+
+		TArrayView<const FVector3f> PositionView = MakeConstArrayView(BoneVerts.Positions);
+		FVector3f Extent;
+		Extent.X = AxisInterval(PositionView, XAxis);
+		Extent.Y = AxisInterval(PositionView, YAxis);
+		Extent.Z = AxisInterval(PositionView, ZAxis);
+
+		// Work out length and radius, length is the longest extent, radius is half the diagonal of the other two.
+		float Length, Radius;
+		if(Extent.X > Extent.Y && Extent.X > Extent.Z)
+		{
+			Length = Extent.X;
+			Radius = FMath::Sqrt(Extent.Y * Extent.Y + Extent.Z * Extent.Z);
+		}
+		else if(Extent.Y > Extent.Z)
+		{
+			Length = Extent.Y;
+			Radius = FMath::Sqrt(Extent.X * Extent.X + Extent.Z * Extent.Z);
+		}
+		else
+		{
+			Length = Extent.Z;
+			Radius = FMath::Sqrt(Extent.X * Extent.X + Extent.Y * Extent.Y);
+		}
+
+		FMatrix ElemMatrix(static_cast<FVector>(XAxis),
+			static_cast<FVector>(YAxis),
+			static_cast<FVector>(ZAxis),
+			FVector::ZeroVector);
+
+		TSharedPtr<FKSphylElem> NewCapsule = MakeShared<FKSphylElem>();
+		NewCapsule->Center = static_cast<FVector>(AveragePosition(BoneVerts.Positions));
+		NewCapsule->Length = Length;
+		NewCapsule->Radius = Radius * 0.5f;
+		NewCapsule->Rotation = ElemMatrix.Rotator();
+
+		Result.AddDefaulted();
+		Result.Last().GeometryElement = NewCapsule;
+	}
+
+	return Result;
+}
+
+#if WITH_EDITOR
+bool UBoneGeometryGenerator_Capsule::CanDebugDraw() const
+{
+	return true;
+}
+
+bool UBoneGeometryGenerator_Capsule::CanDebugDrawViewMode(const FName& ViewModeName) const
+{
+	return true;
 }
 
+void UBoneGeometryGenerator_Capsule::DebugDraw(UE::Dataflow::FContext& Context, IDataflowDebugDrawInterface& DataflowRenderingInterface, const FDataflowNode::FDebugDrawParameters& DebugDrawParameters, FRigidAssetBoneSelection Bones) const
+{
+	if(!Bones.Mesh || !Bones.Skeleton || !bDrawVerts)
+	{
+		return;
+	}
+
+	TArray<FTransform> BoneTransforms;
+	const FReferenceSkeleton& RefSkel = Bones.Mesh->GetRefSkeleton();
+	RefSkel.GetBoneAbsoluteTransforms(BoneTransforms);
+
+	int32 BoneColorSeed = 0;
+	for(const FRigidAssetBoneInfo& Bone : Bones.SelectedBones)
+	{
+		if(DrawVertsForBoneName != NAME_None && Bone.Name == DrawVertsForBoneName)
+		{
+			continue;
+		}
+
+		FBoneVertData BoneVerts = GetBoneVertexData(Bone.Name, Bones.Mesh, BaseSettings.SourceLod, BaseSettings.VertexMode);
+
+		DataflowRenderingInterface.SetColor(FLinearColor::MakeRandomSeededColor(BoneColorSeed++));
+		for(const FVector3f& Pos : BoneVerts.Positions)
+		{
+			DataflowRenderingInterface.DrawPoint(BoneTransforms[Bone.Index].TransformPosition(static_cast<FVector>(Pos)));
+		}
+	}
+}
+#endif
+
 ////////////////////////////////////////////////////////////////////////////////////////////////////////////////
 
 TArray<UE::Chaos::RigidAsset::FSimpleGeometry> UBoneGeometryGenerator_Convex::Build(FRigidAssetBoneSelection Bones)
 {
-	return {};
+	if(!Bones.Mesh || !Bones.Skeleton)
+	{
+		return {};
+	}
+
+	TArray<UE::Chaos::RigidAsset::FSimpleGeometry> Result;
+
+	for(const FRigidAssetBoneInfo& Bone : Bones.SelectedBones)
+	{
+		FBoneVertData BoneVerts = GetBoneVertexData(Bone.Name, Bones.Mesh, BaseSettings.SourceLod, BaseSettings.VertexMode);
+
+		if(BoneVerts.Positions.Num() == 0)
+		{
+			// No influences, make a small sphere and warn
+			// #TODO pass context through for dataflow warnings
+			Result.AddDefaulted();
+			Result.Last().GeometryElement = MakeShared<FKSphereElem>(5.0f);
+			continue;
+		}
+
+		FBox3f VertBox(BoneVerts.Positions);
+
+		TSharedPtr<FKConvexElem> NewConvex = MakeShared<FKConvexElem>();
+		TArray<Chaos::FConvexBuilder::FVec3Type> InPositions;
+		TArray<Chaos::FConvexBuilder::FPlaneType> Planes;
+		TArray<TArray<int32>> Indices;
+		TArray<Chaos::FConvexBuilder::FVec3Type> ConvexVerts;
+		Chaos::FConvexBuilder::FAABB3Type Aabb;
+
+		InPositions.Reserve(BoneVerts.Positions.Num());
+		for(const FVector3f& Vert : BoneVerts.Positions)
+		{
+			InPositions.Add(Vert);
+		}
+
+		Chaos::FConvexBuilder::Build(InPositions, Planes, Indices, ConvexVerts, Aabb);
+
+		NewConvex->VertexData.Reset(ConvexVerts.Num());
+
+		for(const Chaos::FConvexBuilder::FVec3Type& ConvexVert : ConvexVerts)
+		{
+			NewConvex->VertexData.Add(static_cast<FVector>(ConvexVert));
+		}
+
+		NewConvex->UpdateElemBox();
+
+		Result.AddDefaulted();
+		Result.Last().GeometryElement = NewConvex;
+	}
+
+	return Result;
 }
diff --git a/Engine/Plugins/Experimental/ChaosRigidAsset/Source/ChaosRigidAssetNodes/Private/Generators/BoneGeometryGenerators.h b/Engine/Plugins/Experimental/ChaosRigidAsset/Source/ChaosRigidAssetNodes/Private/Generators/BoneGeometryGenerators.h
--- a/Engine/Plugins/Experimental/ChaosRigidAsset/Source/ChaosRigidAssetNodes/Private/Generators/BoneGeometryGenerators.h
+++ b/Engine/Plugins/Experimental/ChaosRigidAsset/Source/ChaosRigidAssetNodes/Private/Generators/BoneGeometryGenerators.h
@@ -43,6 +43,20 @@ enum class EVertexSelectMode
 
 ////////////////////////////////////////////////////////////////////////////////////////////////////////////////
 
+/**
+ * How to align a body during generation
+ * Bone will align along the bone direction and wrap the shape around the verts
+ * Verts will align along the major axis of the selected verts
+ */
+UENUM()
+enum class EBodyAlignment
+{
+	Bone,
+	Verts
+};
+
+////////////////////////////////////////////////////////////////////////////////////////////////////////////////
+
 /**
  * Generic configuration for all generator types
  */
@@ -91,6 +105,12 @@ public:
 		return {};
 	}
 
+#if WITH_EDITOR
+	virtual bool CanDebugDraw() const;
+	virtual bool CanDebugDrawViewMode(const FName& ViewModeName) const;
+	virtual void DebugDraw(UE::Dataflow::FContext& Context, IDataflowDebugDrawInterface& DataflowRenderingInterface, const FDataflowNode::FDebugDrawParameters& DebugDrawParameters, FRigidAssetBoneSelection Bones) const;
+#endif
+
 protected:
 
 	// Base settings for all generation types
@@ -131,10 +151,24 @@ public:
 
 	TArray<UE::Chaos::RigidAsset::FSimpleGeometry> Build(FRigidAssetBoneSelection Bones) final override;
 
+#if WITH_EDITOR
+	virtual bool CanDebugDraw() const override;
+	virtual bool CanDebugDrawViewMode(const FName& ViewModeName) const override;
+	virtual void DebugDraw(UE::Dataflow::FContext& Context, IDataflowDebugDrawInterface& DataflowRenderingInterface, const FDataflowNode::FDebugDrawParameters& DebugDrawParameters, FRigidAssetBoneSelection Bones) const override;
+#endif
+
 private:
 
 	UPROPERTY(EditAnywhere, Category = Settings)
 	bool bUseGeometry = true;
+
+#if WITH_EDITORONLY_DATA
+	UPROPERTY(EditAnywhere, Category=Debug)
+	bool bDrawVerts = true;
+
+	UPROPERTY(EditAnywhere, Category=Debug)
+	FName DrawVertsForBoneName = NAME_None;
+#endif
 };
 
 ////////////////////////////////////////////////////////////////////////////////////////////////////////////////
@@ -152,7 +186,24 @@ public:
 
 	TArray<UE::Chaos::RigidAsset::FSimpleGeometry> Build(FRigidAssetBoneSelection Bones) final override;
 
+#if WITH_EDITOR
+	virtual bool CanDebugDraw() const override;
+	virtual bool CanDebugDrawViewMode(const FName& ViewModeName) const override;
+	virtual void DebugDraw(UE::Dataflow::FContext& Context, IDataflowDebugDrawInterface& DataflowRenderingInterface, const FDataflowNode::FDebugDrawParameters& DebugDrawParameters, FRigidAssetBoneSelection Bones) const override;
+#endif
+
 private:
+
+	UPROPERTY(EditAnywhere, Category=Geometry)
+	EBodyAlignment Alignment = EBodyAlignment::Bone;
+
+#if WITH_EDITORONLY_DATA
+	UPROPERTY(EditAnywhere, Category=Debug)
+	bool bDrawVerts = true;
+
+	UPROPERTY(EditAnywhere, Category=Debug)
+	FName DrawVertsForBoneName = NAME_None;
+#endif
 };
 
 ////////////////////////////////////////////////////////////////////////////////////////////////////////////////
diff --git a/Engine/Plugins/Experimental/ChaosRigidAsset/Source/ChaosRigidAssetNodes/Private/PhysicsAssetDataflowNodes.cpp b/Engine/Plugins/Experimental/ChaosRigidAsset/Source/ChaosRigidAssetNodes/Private/PhysicsAssetDataflowNodes.cpp
--- a/Engine/Plugins/Experimental/ChaosRigidAsset/Source/ChaosRigidAssetNodes/Private/PhysicsAssetDataflowNodes.cpp
+++ b/Engine/Plugins/Experimental/ChaosRigidAsset/Source/ChaosRigidAssetNodes/Private/PhysicsAssetDataflowNodes.cpp
@@ -518,6 +518,12 @@ void FDataflowCreateGeometryForBones::DebugDraw(UE::Dataflow::FContext& Context,
 	const FRigidAssetBoneSelection& InSelection = GetValue(Context, &Selection);
 	TRefCountPtr<IDataflowDebugDrawObject> DrawObject(MakeDebugDrawObject<FDataflowDebugDrawBoneSelectionObject>(DataflowRenderingInterface.ModifyDataflowElements(), InSelection));
 	DataflowRenderingInterface.DrawObject(DrawObject);
+
+	TObjectPtr<UBoneGeometryGenerator> InGenerator = GetValue(Context, &Generator);
+	if(InGenerator && InGenerator->CanDebugDraw())
+	{
+		InGenerator->DebugDraw(Context, DataflowRenderingInterface, DebugDrawParameters, InSelection);
+	}
 }
 #endif
 
diff --git a/Engine/Source/Runtime/Experimental/Chaos/Private/GeometryCollection/Facades/CollectionRenderingFacade.cpp b/Engine/Source/Runtime/Experimental/Chaos/Private/GeometryCollection/Facades/CollectionRenderingFacade.cpp
--- a/Engine/Source/Runtime/Experimental/Chaos/Private/GeometryCollection/Facades/CollectionRenderingFacade.cpp
+++ b/Engine/Source/Runtime/Experimental/Chaos/Private/GeometryCollection/Facades/CollectionRenderingFacade.cpp
@@ -632,24 +632,25 @@ namespace GeometryCollection::Facades
 		return RetCount;
 	}
 
-	void FRenderingFacade::AddCapsule(const float Length, const float Radius)
+	void FRenderingFacade::AddCapsule(const float Length, const float Radius, FLinearColor Color, int32 Sides)
 	{
 		UE::Geometry::FCapsuleGenerator Generator;
-		Generator.Radius = Radius;
-		Generator.SegmentLength = Length;
+		Generator.Radius = FMath::Max(UE_SMALL_NUMBER, Radius);
+		Generator.SegmentLength = FMath::Max(UE_SMALL_NUMBER, Length);
+		Generator.NumCircleSteps = FMath::Max(3, Sides);
 		Generator.Generate();
 
 		TArray<FVector3f> NarrowVerts;
+		TArray<FLinearColor> Colors;
 		TArray<FIntVector> Triangles;
-		TArray<FLinearColor> Colours;
-		NarrowVerts.Reserve(Generator.Vertices.Num());
 		NarrowVerts.Reserve(Generator.Vertices.Num());
+		Colors.Reserve(Generator.Vertices.Num());
 		Triangles.Reserve(Generator.Triangles.Num());
 
 		for(const FVector3d& WideVertex : Generator.Vertices)
 		{
 			NarrowVerts.Add((FVector3f)WideVertex);
-			Colours.Add(FLinearColor::White);
+			Colors.Add(Color);
 		}
 
 		for(const UE::Geometry::FIndex3i& TriIndices : Generator.Triangles)
@@ -657,7 +658,7 @@ namespace GeometryCollection::Facades
 			Triangles.Add({TriIndices.A, TriIndices.B, TriIndices.C});
 		}
 
-		AddSurface(MoveTemp(NarrowVerts), MoveTemp(Triangles), MoveTemp(Generator.Normals), MoveTemp(Colours));
+		AddSurface(MoveTemp(NarrowVerts), MoveTemp(Triangles), MoveTemp(Generator.Normals), MoveTemp(Colors));
 	}
 
 	void FRenderingFacade::AddFaces(const TArray<FVector3f>& InVertices, TArray<FIntVector>& InIndices, TArray<FLinearColor>& InColors)
diff --git a/Engine/Source/Runtime/Experimental/Chaos/Public/GeometryCollection/Facades/CollectionRenderingFacade.h b/Engine/Source/Runtime/Experimental/Chaos/Public/GeometryCollection/Facades/CollectionRenderingFacade.h
--- a/Engine/Source/Runtime/Experimental/Chaos/Public/GeometryCollection/Facades/CollectionRenderingFacade.h
+++ b/Engine/Source/Runtime/Experimental/Chaos/Public/GeometryCollection/Facades/CollectionRenderingFacade.h
@@ -76,7 +76,7 @@ namespace GeometryCollection::Facades
 		CHAOS_API void AddSpheres(const TArray<FSphere>& InSpheres, const FLinearColor& InColor);
 
 		/** Add a capsule to the rendering view */
-		CHAOS_API void AddCapsule(const float Length, const float Radius);
+		CHAOS_API void AddCapsule(const float Length, const float Radius, FLinearColor Color = FLinearColor::White, int32 Sides = 8);
 
 		/** Add a surface to the rendering view.*/
 		CHAOS_API void AddSurface(TArray<FVector3f>&& InVertices, TArray<FIntVector>&& InIndices, TArray<FVector3f>&& InNormals, TArray<FLinearColor>&& InColors);
