[Dataflow] Fixed physics asset generation not correctly setting base joint names on constraints leading to some systems not being able to find bones linked to constraints

diff --git a/Engine/Plugins/Experimental/ChaosRigidAsset/Source/ChaosRigidAssetNodes/Private/Generators/ConstraintGenerators.cpp b/Engine/Plugins/Experimental/ChaosRigidAsset/Source/ChaosRigidAssetNodes/Private/Generators/ConstraintGenerators.cpp
--- a/Engine/Plugins/Experimental/ChaosRigidAsset/Source/ChaosRigidAssetNodes/Private/Generators/ConstraintGenerators.cpp
+++ b/Engine/Plugins/Experimental/ChaosRigidAsset/Source/ChaosRigidAssetNodes/Private/Generators/ConstraintGenerators.cpp
@@ -54,6 +54,7 @@ TArray<TObjectPtr<UPhysicsConstraintTemplate>> UConstraintGenerator_SwingTwist::
 			Defaults.SetAngularSwing2Limit(EAngularConstraintMotion::ACM_Limited, Swing2Limit);
 			Defaults.SetAngularTwistLimit(EAngularConstraintMotion::ACM_Limited, TwistLimit);
 
+			Defaults.JointName = Bone.Name;
 			Defaults.ConstraintBone1 = Bone.Name;
 			Defaults.ConstraintBone2 = ParentBone->Name;
 
