Rigid body asset Dataflow nodes :
- add a node to run the physics asset state in a RBAN simulation and cache the result in an animation flle
- add a animsequence terminal node

diff --git a/Engine/Plugins/Experimental/ChaosRigidAsset/Source/ChaosRigidAssetNodes/Private/PhysicsAssetDataflowNodes.cpp b/Engine/Plugins/Experimental/ChaosRigidAsset/Source/ChaosRigidAssetNodes/Private/PhysicsAssetDataflowNodes.cpp
--- a/Engine/Plugins/Experimental/ChaosRigidAsset/Source/ChaosRigidAssetNodes/Private/PhysicsAssetDataflowNodes.cpp
+++ b/Engine/Plugins/Experimental/ChaosRigidAsset/Source/ChaosRigidAssetNodes/Private/PhysicsAssetDataflowNodes.cpp
@@ -13,6 +13,14 @@
 #include "PhysicsAssetBuilder.h"
 #include "Engine/SkeletalMesh.h"
 #include "ReferenceSkeleton.h"
+#include "Physics/ImmediatePhysics/ImmediatePhysicsAdapters.h"
+#include "Physics/ImmediatePhysics/ImmediatePhysicsChaos/ImmediatePhysicsSimulation_Chaos.h"
+#include "Physics/ImmediatePhysics/ImmediatePhysicsChaos/ImmediatePhysicsActorHandle_Chaos.h"
+#include "Physics/ImmediatePhysics/ImmediatePhysicsChaos/ImmediatePhysicsJointHandle_Chaos.h"
+#include "Physics/ImmediatePhysics/ImmediatePhysicsDeclares.h"
+#include "Animation/AnimSequence.h"
+#include "AnimationRuntime.h"
+#include "AnimationUtils.h"
 
 FDataflowPhysicsAssetTerminalNode::FDataflowPhysicsAssetTerminalNode(const UE::Dataflow::FNodeParameters& InParam, FGuid InGuid /*= FGuid::NewGuid()*/)
 	: FDataflowTerminalNode(InParam, InGuid)
@@ -43,12 +51,12 @@ void FDataflowPhysicsAssetTerminalNode::SetAssetValue(TObjectPtr<UObject> Asset,
 
 void FDataflowPhysicsAssetTerminalNode::Evaluate(UE::Dataflow::FContext& Context) const
 {
-
 }
 
 void UE::Dataflow::RegisterPhysicsAssetTerminalNode()
 {
 	DATAFLOW_NODE_REGISTER_CREATION_FACTORY(FDataflowPhysicsAssetTerminalNode);
+	DATAFLOW_NODE_REGISTER_CREATION_FACTORY(FDataflowAnimSequenceAssetTerminalNode);
 }
 
 void UE::Dataflow::RegisterPhysicsAssetNodes()
@@ -75,6 +83,9 @@ void UE::Dataflow::RegisterPhysicsAssetNodes()
 	DATAFLOW_NODE_REGISTER_CREATION_FACTORY(FDataflowSelectBonesByName);
 	DATAFLOW_NODE_REGISTER_CREATION_FACTORY(FDataflowSelectConnectedBones);
 
+	DATAFLOW_NODE_REGISTER_CREATION_FACTORY(FDataflowRigidSimulationCachingNode);
+	
+
 	RegisterBoneGeometryGeneratorNodes();
 	RegisterConstraintGeneratorNodes();
 }
@@ -518,12 +529,6 @@ void FDataflowCreateGeometryForBones::DebugDraw(UE::Dataflow::FContext& Context,
 	const FRigidAssetBoneSelection& InSelection = GetValue(Context, &Selection);
 	TRefCountPtr<IDataflowDebugDrawObject> DrawObject(MakeDebugDrawObject<FDataflowDebugDrawBoneSelectionObject>(DataflowRenderingInterface.ModifyDataflowElements(), InSelection));
 	DataflowRenderingInterface.DrawObject(DrawObject);
-
-	TObjectPtr<UBoneGeometryGenerator> InGenerator = GetValue(Context, &Generator);
-	if(InGenerator && InGenerator->CanDebugDraw())
-	{
-		InGenerator->DebugDraw(Context, DataflowRenderingInterface, DebugDrawParameters, InSelection);
-	}
 }
 #endif
 
@@ -837,6 +842,290 @@ void FDataflowAggGeomAddShape::Evaluate(UE::Dataflow::FContext& Context, const F
 	}
 }
 
+////////////////////////////////////////////////////////////////////////////////////////////////
+
+void FDataflowRigidSimulationCachingNode::Register(const UE::Dataflow::FNodeParameters& InParam)
+{
+	AddInputs(State, Duration, FrameRate, Gravity, Velocity);
+	AddOutputs(Animation);
+}
+
+void FDataflowRigidSimulationCachingNode::Evaluate(UE::Dataflow::FContext& Context, const FDataflowOutput* Out) const
+{
+	using namespace ImmediatePhysics;
+
+	if (Out->IsA(&Animation))
+	{
+		const FPhysicsAssetDataflowState& InState = GetValue(Context, &State);
+		const float InDuration= GetValue(Context, &Duration);
+		int32 InFrameRate = GetValue(Context, &FrameRate);
+		if (InFrameRate <= 0)
+		{
+			InFrameRate = 30;
+		}
+		const FVector InGravity= GetValue(Context, &Gravity);
+		const FVector InVelocity = GetValue(Context, &Velocity);
+		
+
+		const float FixedTimeStep = 1.0f / static_cast<float>(InFrameRate);
+		const int32 NumFrames = FMath::Max(1, (int32)(InDuration * InFrameRate));
+		
+		TObjectPtr<UAnimSequence> OutAnim = NewObject<UAnimSequence>();
+#if WITH_EDITOR
+		if (InState.TargetSkeleton && InState.TargetMesh)
+		{
+			OutAnim->SetSkeleton(InState.TargetSkeleton);
+			OutAnim->BoneCompressionSettings = FAnimationUtils::GetDefaultAnimationRecorderBoneCompressionSettings();
+
+			IAnimationDataController& Controller = OutAnim->GetController();
+			Controller.SetModel(OutAnim->GetDataModelInterface());
+			Controller.InitializeModel();
+			Controller.RemoveAllCurvesOfType(ERawCurveTrackTypes::RCT_Float);
+			Controller.RemoveAllBoneTracks(/*bTransactRecording*/false);
+			OutAnim->ResetAnimation();
+
+			//Set Interpolation type (Step or Linear), doesn't look like there is a controller for this.
+			OutAnim->Interpolation = EAnimInterpolationType::Linear;
+			const FFrameRate RecordingRate = FFrameRate(InFrameRate, 1);
+			// Set frame rate and number of frames
+			Controller.SetFrameRate(RecordingRate, /*bTransactRecording*/false);
+			Controller.SetNumberOfFrames(NumFrames, /*bTransactRecording*/false);
+
+			OutAnim->RetargetSource = InState.TargetSkeleton->GetRetargetSourceForMesh(InState.TargetMesh);
+
+			FSimulationState SimState;
+			SimState.FixedTimeStep = FixedTimeStep;
+			BuildSimulationFromState(InState, SimState);
+			if (SimState.Simulation.IsValid())
+			{
+				TMap<FName, FRawAnimSequenceTrack> RawTracksByName;
+				TMap<FName, FActorHandle*> ParentActorByChildName;
+
+				const FReferenceSkeleton& RefSkeleton = InState.TargetMesh->GetRefSkeleton();
+
+				// create the tracks from the simulation actors
+				for (const TPair<FName, FActorHandle*>& ActorByName : SimState.ActorsByName)
+				{
+					if (ActorByName.Value)
+					{
+						const FName ActorName = ActorByName.Key;
+						Controller.AddBoneCurve(ActorName, /*bTransactRecording*/false);
+						RawTracksByName.Add(ActorName, FRawAnimSequenceTrack());
+
+						const int32 BoneIndex = RefSkeleton.FindBoneIndex(ActorName);
+						const int32 ParentBoneIndex = (BoneIndex == INDEX_NONE) ? INDEX_NONE : RefSkeleton.GetParentIndex(BoneIndex);
+						if (ParentBoneIndex != INDEX_NONE)
+						{
+							const FName ParentName = RefSkeleton.GetBoneName(ParentBoneIndex);
+							if (FActorHandle* ParentHandle = SimState.ActorsByName.FindRef(ParentName))
+							{
+								ParentActorByChildName.Add(ActorName, ParentHandle);
+							}
+						}
+					}
+				}
+
+				for (int32 Frame = 0; Frame < NumFrames; ++Frame)
+				{
+					// record results
+					for (const TPair<FName, FActorHandle*>& ActorByName : SimState.ActorsByName)
+					{
+						if (FActorHandle* Handle = ActorByName.Value)
+						{
+							const FName ActorName = ActorByName.Key;
+
+							FTransform ActorLocalTransform = Handle->GetWorldTransform();
+							if (const FActorHandle* ParentHandle = ParentActorByChildName.FindRef(ActorName))
+							{
+								const FTransform ParentActorTransform = ParentHandle->GetWorldTransform();
+								ActorLocalTransform.SetToRelativeTransform(ParentActorTransform);
+							}
+
+							if (FRawAnimSequenceTrack* RawTrack = RawTracksByName.Find(ActorName))
+							{
+								RawTrack->PosKeys.Add(FVector3f(ActorLocalTransform.GetTranslation()));
+								RawTrack->RotKeys.Add(FQuat4f(ActorLocalTransform.GetRotation()));
+								RawTrack->ScaleKeys.Add(FVector3f(ActorLocalTransform.GetScale3D()));
+							}
+
+							const float Ratio = FMath::Clamp(Handle->GetParticle()->GetX().Z / 2800.f, 0, 1);
+							Handle->SetLinearVelocity(InVelocity * Ratio * Ratio);
+						}
+					}
+
+					SimState.Simulation->UpdateSimulationSpace(FTransform::Identity, InVelocity, FVector::ZeroVector, FVector::ZeroVector, FVector::ZeroVector);
+					RunSimulation(SimState, InGravity);
+				}
+
+				// done let's save the tracks 
+				for (const TPair<FName, FRawAnimSequenceTrack>& TrackByName : RawTracksByName)
+				{
+					const FRawAnimSequenceTrack& RawTrack = TrackByName.Value;
+					const FName TrackName = TrackByName.Key;
+					Controller.SetBoneTrackKeys(TrackName, RawTrack.PosKeys, RawTrack.RotKeys, RawTrack.ScaleKeys, /*bTransactRecording*/false);
+				}
+
+				OutAnim->PostEditChange();
+			}
+		}
+		else
+		{
+			Context.Warning("Input State TargetSkeleton or TargetMesh is null, outputing an empty animation sequence", this, Out);
+		}
+#endif // WITH_EDITOR
+		SetValue(Context, OutAnim, &Animation);
+	}
+}
+void FDataflowRigidSimulationCachingNode::RunSimulation(FSimulationState& SimulationState, const FVector& InGravity) const
+{
+	if (SimulationState.Simulation)
+	{
+		SimulationState.Simulation->Simulate_AssumesLocked(SimulationState.FixedTimeStep, SimulationState.FixedTimeStep, 1, InGravity);
+	}
+}
+
+void FDataflowRigidSimulationCachingNode::BuildSimulationFromState(const FPhysicsAssetDataflowState& InState, FSimulationState& OutSimState) const
+{
+	using namespace ImmediatePhysics;
+
+	if (InState.TargetMesh == nullptr)
+	{
+		return;
+	}
+
+	OutSimState.Simulation = MakeUnique<FSimulation>();
+#if WITH_CHAOS_VISUAL_DEBUGGER
+	OutSimState.Simulation->GetChaosVDContextData().Id = FChaosVDRuntimeModule::Get().GenerateUniqueID();
+	OutSimState.Simulation->GetChaosVDContextData().Type = static_cast<int32>(EChaosVDContextType::Solver);
+#endif
+
+	// structure to efficiently compute bones component space
+	TArray<FTransform> CachedTransforms;
+	TArray<bool> CachedTransformReady;
+	const FReferenceSkeleton& RefSkeleton = InState.TargetMesh->GetRefSkeleton();
+	CachedTransforms.SetNumUninitialized(RefSkeleton.GetRefBonePose().Num());
+	CachedTransformReady.SetNumZeroed(RefSkeleton.GetRefBonePose().Num());
+
+
+	for (TObjectPtr<USkeletalBodySetup> BodySetup : InState.Bodies)
+	{
+		TUniquePtr<FBodyInstance> BodyInstance = CreateBodyInstance(BodySetup);
+		if (BodyInstance)
+		{
+			const int32 BoneIndex = RefSkeleton.FindBoneIndex(BodySetup->BoneName);
+			FTransform BoneComponentSpaceTransform{ FTransform::Identity };
+			if (BoneIndex != INDEX_NONE)
+			{
+				BoneComponentSpaceTransform = FAnimationRuntime::GetComponentSpaceTransformWithCache(RefSkeleton, RefSkeleton.GetRefBonePose(), BoneIndex, CachedTransforms, CachedTransformReady);
+			}
+
+			const bool bSimulated = (BodySetup->PhysicsType == EPhysicsType::PhysType_Simulated || BodySetup->PhysicsType == EPhysicsType::PhysType_Default);
+			EActorType ActorType = bSimulated ? EActorType::DynamicActor : EActorType::KinematicActor;
+			FActorHandle* NewBodyHandle = OutSimState.Simulation->CreateActor(MakeActorSetup(ActorType, BodyInstance.Get(), BodyInstance->GetUnrealWorldTransform()));
+			if (NewBodyHandle)
+			{
+				// This is a local simulation we assume that world space is the same as component space
+				NewBodyHandle->InitWorldTransform(BoneComponentSpaceTransform);
+				NewBodyHandle->SetName(BodySetup->BoneName);
+				NewBodyHandle->SetEnabled(true);
+				NewBodyHandle->SetHasCollision(false);
+				//OutSimState.Simulation->AddToCollidingPairs(NewBodyHandle);
+				OutSimState.ActorsByName.Add(BodySetup->BoneName, NewBodyHandle);
+				UE_LOG(LogTemp, Warning, TEXT("@@ Create Sim Actor : [%s] - bSimulated[%d] - Transform[%s]"), *BodySetup->BoneName.ToString(), (int)bSimulated, *BoneComponentSpaceTransform.ToString());
+			}
+		}
+	}
+
+
+	TArray<FSimulation::FIgnorePair> IgnoreCollisionPairs;
+	for (TObjectPtr<UPhysicsConstraintTemplate> ConstraintTemplate: InState.Constraints)
+	{
+		if (ConstraintTemplate)
+		{
+			FActorHandle* Actor1 = OutSimState.ActorsByName.FindRef(ConstraintTemplate->DefaultInstance.ConstraintBone1, nullptr);
+			FActorHandle* Actor2 = OutSimState.ActorsByName.FindRef(ConstraintTemplate->DefaultInstance.ConstraintBone2, nullptr);
+			if (Actor1 && Actor2)
+			{
+				OutSimState.Simulation->CreateJoint(MakeJointSetup(&ConstraintTemplate->DefaultInstance, Actor1, Actor2));
+				FSimulation::FIgnorePair Pair;
+				Pair.A = Actor1;
+				Pair.B = Actor2;
+				IgnoreCollisionPairs.Add(Pair);
+
+			}
+		}
+	}
+
+	OutSimState.Simulation->SetIgnoreCollisionPairTable(IgnoreCollisionPairs);
+	//OutSimState.Simulation->SetIgnoreCollisionActors(IgnoreCollisionActors);
+	OutSimState.Simulation->SetSolverSettings(
+		OutSimState.FixedTimeStep,
+		/*CullDistance*/ 10000000,
+		/*MaxDepenetrationVelocity*/ 100000000,
+		/*bUseLinearJointSolver*/ true,
+		/*PositionIterations*/ 8,
+		/*VelocityIterations*/ 2,
+		/*SolverSettings.*/ 1,
+		/*bUseManifolds*/ true);
+}
+
+TUniquePtr<FBodyInstance> FDataflowRigidSimulationCachingNode::CreateBodyInstance(USkeletalBodySetup* BodySetup) const
+{
+	TUniquePtr<FBodyInstance> BodyInstance;
+	if (BodySetup && BodySetup->AggGeom.GetElementCount() > 0)
+	{
+		BodyInstance = MakeUnique<FBodyInstance>();
+		BodyInstance->bAutoWeld = false;
+		BodyInstance->bSimulatePhysics = false;
+		BodyInstance->InitBody(BodySetup, FTransform::Identity, nullptr, nullptr);
+
+		if (!BodyInstance->IsValidBodyInstance() || !ensure(BodyInstance->GetBodySetup()))
+		{
+			// Some part of the initialization process failed.
+			BodyInstance.Reset();
+		}
+	}
+	return BodyInstance;
+}
+
+////////////////////////////////////////////////////////////////////////////////////////////////
+FDataflowAnimSequenceAssetTerminalNode::FDataflowAnimSequenceAssetTerminalNode(const UE::Dataflow::FNodeParameters& InParam, FGuid InGuid)
+	: FDataflowTerminalNode(InParam, InGuid)
+{
+	RegisterInputConnection(&Animation);
+	RegisterInputConnection(&AnimationAsset);
+}
+
+void FDataflowAnimSequenceAssetTerminalNode::SetAssetValue(TObjectPtr<UObject> Asset, UE::Dataflow::FContext& Context) const
+{
+	UAnimSequence* AnimTarget = Asset ? Cast<UAnimSequence>(Asset.Get()) : nullptr;
+	if (!AnimTarget)
+	{
+		AnimTarget = GetValue(Context, &AnimationAsset);
+	}
+
+	if (AnimTarget)
+	{
+#if WITH_EDITOR
+		// This should be const but CreateAnimation does not take a const parameter
+		if (UAnimSequence* InAnimation = GetValue(Context, &Animation))
+		{
+			// copied from UAnimSequence::CreateAnimation because the method does not take a const anim sequence as a parameter
+			AnimTarget->CreateAnimation(InAnimation);
+		}
+
+		AnimTarget->Modify();
+		AnimTarget->PostEditChange();
+#endif // WITH_EDITOR
+	}
+}
+
+void FDataflowAnimSequenceAssetTerminalNode::Evaluate(UE::Dataflow::FContext& Context) const
+{
+}
+
+////////////////////////////////////////////////////////////////////////////////////////////////
+
 void FDataflowAggGeomAddShape::Register()
 {
 	AddInputs(AggGeom);
diff --git a/Engine/Plugins/Experimental/ChaosRigidAsset/Source/ChaosRigidAssetNodes/Private/PhysicsAssetDataflowNodes.h b/Engine/Plugins/Experimental/ChaosRigidAsset/Source/ChaosRigidAssetNodes/Private/PhysicsAssetDataflowNodes.h
--- a/Engine/Plugins/Experimental/ChaosRigidAsset/Source/ChaosRigidAssetNodes/Private/PhysicsAssetDataflowNodes.h
+++ b/Engine/Plugins/Experimental/ChaosRigidAsset/Source/ChaosRigidAssetNodes/Private/PhysicsAssetDataflowNodes.h
@@ -16,12 +16,16 @@
 #include "RigidDataflowNode.h"
 #include "ShapeElemNodes.h"
 
+#include "Physics/ImmediatePhysics/ImmediatePhysicsDeclares.h"
+
 #include "PhysicsAssetDataflowNodes.generated.h"
 
 class USkeleton;
 class USkeletalMesh;
 class UPhysicsAsset;
 class USkeletalBodySetup;
+class UPhysicsConstraintTemplate;
+class UAnimSequence;
 
 ////////////////////////////////////////////////////////////////////////////////////////////////////////////////
 
@@ -522,6 +526,93 @@ private:
 
 ////////////////////////////////////////////////////////////////////////////////////////////////////////////////
 
+/**
+* Run a rigid simulation of a Physics asset state and bake it into an animation file 
+*/
+USTRUCT(meta = (Experimental))
+struct FDataflowRigidSimulationCachingNode : public FRigidDataflowNode
+{
+	GENERATED_BODY()
+	DATAFLOW_NODE_DEFINE_INTERNAL(FDataflowRigidSimulationCachingNode, "Rigid Simulation Caching", "PhysicsAsset", "")
+
+public:
+
+	FDataflowRigidSimulationCachingNode(const UE::Dataflow::FNodeParameters& InParam, FGuid InGuid = FGuid::NewGuid())
+		: FRigidDataflowNode(InParam, InGuid)
+	{
+		Register(InParam);
+	}
+
+private:
+
+	// Input physics asset state to use for the simulation
+	UPROPERTY(meta = (DataflowInput))
+	FPhysicsAssetDataflowState State;
+
+	/** OUtput animation asset  */
+	UPROPERTY(meta = (DataflowOutput))
+	TObjectPtr<UAnimSequence> Animation = nullptr;
+
+	/** Duration in seconds of the simulation */
+	UPROPERTY(EditAnywhere, Category = Caching, meta = (DataflowInput, ClampMin=0))
+	float Duration = 5.0;
+
+	/** Number of frames per second */
+	UPROPERTY(EditAnywhere, Category = Caching, meta = (DataflowInput, ClampMin = 1))
+	int32 FrameRate = 30;
+
+	/** Gravity vector in cm/s2 */
+	UPROPERTY(EditAnywhere, Category = Simulation, meta = (DataflowInput))
+	FVector Gravity = FVector(0, 0, 980);
+
+	/** simulation space velocity in cm/s */
+	UPROPERTY(EditAnywhere, Category = Simulation, meta = (DataflowInput))
+	FVector Velocity = FVector(0, 0, 0);
+
+	void Evaluate(UE::Dataflow::FContext& Context, const FDataflowOutput* Out) const override;
+	void Register(const UE::Dataflow::FNodeParameters& InParam);
+
+	TUniquePtr<FBodyInstance> CreateBodyInstance(USkeletalBodySetup* BodySetup) const;
+
+	struct FSimulationState
+	{
+		float FixedTimeStep = 0;
+		TUniquePtr<ImmediatePhysics::FSimulation> Simulation;
+		TMap<FName, ImmediatePhysics::FActorHandle*> ActorsByName;
+	};
+
+	void BuildSimulationFromState(const FPhysicsAssetDataflowState& InState, FSimulationState& OutSimState) const;
+	void RunSimulation(FSimulationState& Simulation, const FVector& InGravity) const;
+};
+
+
+/**
+ * Terminal node that writes to a animation sequence asset
+ */
+USTRUCT(meta = (DataflowTerminal))
+struct FDataflowAnimSequenceAssetTerminalNode : public FDataflowTerminalNode
+{
+	GENERATED_BODY()
+	DATAFLOW_NODE_DEFINE_INTERNAL(FDataflowAnimSequenceAssetTerminalNode, "AnimSequenceTerminal", "Terminal", "")
+
+public:
+
+	FDataflowAnimSequenceAssetTerminalNode(const UE::Dataflow::FNodeParameters& InParam, FGuid InGuid = FGuid::NewGuid());
+
+private:
+
+	UPROPERTY(meta = (DataflowInput))
+	TObjectPtr<UAnimSequence> Animation;
+
+	UPROPERTY(EditAnywhere, Category = Asset, meta = (DataflowInput))
+	TObjectPtr<UAnimSequence> AnimationAsset;
+
+	void SetAssetValue(TObjectPtr<UObject> Asset, UE::Dataflow::FContext& Context) const override;
+	void Evaluate(UE::Dataflow::FContext& Context) const override;
+};
+
+////////////////////////////////////////////////////////////////////////////////////////////////////////////////
+
 namespace UE::Dataflow
 {
 	void RegisterPhysicsAssetTerminalNode();
