Fixed crash when attempting to open a dataflow asset without setting a graph for physics assets

diff --git a/Engine/Plugins/Experimental/ChaosRigidAsset/Source/ChaosRigidAssetEditor/Private/ChaosRigidAssetEditorModule.cpp b/Engine/Plugins/Experimental/ChaosRigidAsset/Source/ChaosRigidAssetEditor/Private/ChaosRigidAssetEditorModule.cpp
--- a/Engine/Plugins/Experimental/ChaosRigidAsset/Source/ChaosRigidAssetEditor/Private/ChaosRigidAssetEditorModule.cpp
+++ b/Engine/Plugins/Experimental/ChaosRigidAsset/Source/ChaosRigidAssetEditor/Private/ChaosRigidAssetEditorModule.cpp
@@ -10,6 +10,8 @@
 #include "DataflowRendering.h"
 #include "Dataflow/DataflowRenderingFactory.h"
 #include "Features/IModularFeatures.h"
+#include "Widgets/Notifications/SNotificationList.h"
+#include "Framework/Notifications/NotificationManager.h"
 
 IMPLEMENT_MODULE(FChaosRigidAssetEditorModule, ChaosRigidAssetEditor);
 
@@ -19,7 +21,7 @@ static FAutoConsoleVariableRef GRigidAssetModuleOverridePhysicsAssetEditorCVar(
 	GRigidAssetModuleOverridePhysicsAssetEditor,
 	TEXT("Whether to enable the editor override for physics asset such that physics assets using dataflow will no longer open in the base physics asset editor, preferring the dataflow editor instead"));
 
-/**
+/** 
  * Override for physics asset editor to instead route to the dataflow editor if there
  * is a dataflow attachment present in the asset.
  */
@@ -95,13 +97,23 @@ void SpawnEditorFor(UPhysicsAsset* InAsset)
 			Attachment->GetDataflowInstance().SetDataflowAsset(Cast<UDataflow>(UE::DataflowAssetDefinitionHelpers::NewOrOpenDataflowAsset(InAsset)));
 		}
 
-		UAssetEditorSubsystem* const AssetEditorSubsystem = GEditor->GetEditorSubsystem<UAssetEditorSubsystem>();
+		if(!Attachment->GetDataflowInstance().GetDataflowAsset())
+		{
+			// User opted to not create a dataflow
+			FNotificationInfo NotifyInfo(NSLOCTEXT("RigidAssetEditorModule", "NoDataflow", "No Dataflow asset available to edit."));
+			NotifyInfo.ExpireDuration = 5.0f;
+			FSlateNotificationManager::Get().AddNotification(NotifyInfo);
+		}
+		else
+		{
+			UAssetEditorSubsystem* const AssetEditorSubsystem = GEditor->GetEditorSubsystem<UAssetEditorSubsystem>();
 
-		UDataflowEditor* const AssetEditor = NewObject<UDataflowEditor>(AssetEditorSubsystem, NAME_None, RF_Transient);
-		AssetEditor->RegisterToolCategories({ "General" });
-		const TSubclassOf<AActor> ActorClass = StaticLoadClass(AActor::StaticClass(), nullptr, TEXT("/ChaosRigidAsset/BP_PhysicsAssetPreview.BP_PhysicsAssetPreview_C"), nullptr, LOAD_None, nullptr);
+			UDataflowEditor* const AssetEditor = NewObject<UDataflowEditor>(AssetEditorSubsystem, NAME_None, RF_Transient);
+			AssetEditor->RegisterToolCategories({ "General" });
+			const TSubclassOf<AActor> ActorClass = StaticLoadClass(AActor::StaticClass(), nullptr, TEXT("/ChaosRigidAsset/BP_PhysicsAssetPreview.BP_PhysicsAssetPreview_C"), nullptr, LOAD_None, nullptr);
 
-		AssetEditor->Initialize({ GetDataflowAttachment(InAsset) }, ActorClass);
+			AssetEditor->Initialize({ GetDataflowAttachment(InAsset) }, ActorClass);
+		}
 	}
 }
 
