Dataflow - Add animation bank terminal node and affector nodes
- Also move animation sequence terminal node to the same source files
- Implement a basic Wind Affector as well as an Acceleration Affector

#Synced-CL 48601912

diff --git a/Engine/Plugins/Experimental/ChaosRigidAsset/Source/ChaosRigidAssetNodes/Private/AnimationTerminalNodes.cpp b/Engine/Plugins/Experimental/ChaosRigidAsset/Source/ChaosRigidAssetNodes/Private/AnimationTerminalNodes.cpp
new file mode 100644
--- /dev/null
+++ b/Engine/Plugins/Experimental/ChaosRigidAsset/Source/ChaosRigidAssetNodes/Private/AnimationTerminalNodes.cpp
@@ -0,0 +1,199 @@
+// Copyright Epic Games, Inc. All Rights Reserved.
+
+
+#include "AnimationTerminalNodes.h"
+
+#include "Animation/AnimBank.h"
+#include "Animation/AnimSequence.h"
+#include "Dataflow/DataflowNodeFactory.h"
+#include "Dataflow/DataflowObject.h"
+
+
+void UE::Dataflow::RegisterAnimationTerminalNodes()
+{
+	DATAFLOW_NODE_REGISTER_CREATION_FACTORY(FDataflowAnimBankAssetTerminalNode);
+	DATAFLOW_NODE_REGISTER_CREATION_FACTORY(FDataflowAnimSequenceAssetTerminalNode);
+}
+
+
+////////////////////////////////////////////////////////////////////////////////////////////////
+
+FDataflowAnimSequenceAssetTerminalNode::FDataflowAnimSequenceAssetTerminalNode(const UE::Dataflow::FNodeParameters& InParam, FGuid InGuid)
+	: FDataflowTerminalNode(InParam, InGuid)
+{
+	RegisterInputConnection(&Animation);
+	RegisterInputConnection(&AnimationAsset);
+}
+
+void FDataflowAnimSequenceAssetTerminalNode::SetAssetValue(TObjectPtr<UObject> Asset, UE::Dataflow::FContext& Context) const
+{
+	UAnimSequence* AnimTarget = Asset ? Cast<UAnimSequence>(Asset.Get()) : nullptr;
+	if (!AnimTarget)
+	{
+		AnimTarget = GetValue(Context, &AnimationAsset);
+	}
+
+	if (AnimTarget)
+	{
+#if WITH_EDITOR
+		// This should be const but CreateAnimation does not take a const parameter
+		if (UAnimSequence* InAnimation = GetValue(Context, &Animation))
+		{
+			// copied from UAnimSequence::CreateAnimation because the method does not take a const anim sequence as a parameter
+			AnimTarget->CreateAnimation(InAnimation);
+		}
+
+		AnimTarget->Modify();
+		AnimTarget->PostEditChange();
+#endif // WITH_EDITOR
+	}
+}
+
+void FDataflowAnimSequenceAssetTerminalNode::Evaluate(UE::Dataflow::FContext& Context) const
+{
+}
+
+
+////////////////////////////////////////////////////////////////////////////////////////////////
+
+FDataflowAnimBankAssetTerminalNode::FDataflowAnimBankAssetTerminalNode(const UE::Dataflow::FNodeParameters& InParam, FGuid InGuid)
+	: FDataflowTerminalNode(InParam, InGuid)
+{
+	RegisterInputConnection(&AnimationBankAsset);
+
+	for (int32 Index = 0; Index < NumInitialInputs; ++Index)
+	{
+		AddPins();
+	}
+}
+
+void FDataflowAnimBankAssetTerminalNode::SetAssetValue(TObjectPtr<UObject> Asset, UE::Dataflow::FContext& Context) const
+{
+	UAnimBank* AnimBankTarget = Asset ? Cast<UAnimBank>(Asset.Get()) : nullptr;
+	if (!AnimBankTarget)
+	{
+		AnimBankTarget = GetValue(Context, &AnimationBankAsset);
+	}
+
+	if (AnimBankTarget)
+	{
+#if WITH_EDITORONLY_DATA
+		AnimBankTarget->Sequences.Reset();
+		for (int32 Index = 0; Index < AnimSequences.Num(); ++Index)
+		{
+			TObjectPtr<UAnimSequence> AnimSequence = GetValue(Context, GetConnectionReference(Index), TObjectPtr<UAnimSequence>(nullptr));
+			if (AnimSequence)
+			{
+				FAnimBankSequence& AnimBankSequence = AnimBankTarget->Sequences.AddDefaulted_GetRef();
+				AnimBankSequence.Sequence = AnimSequence;
+				AnimBankSequence.bAutoStart = bAutoStart;
+				AnimBankSequence.bLooping = bLooping;
+				AnimBankSequence.Position = Position;
+				AnimBankSequence.PlayRate = PlayRate;
+				AnimBankSequence.BoundsScale = BoundsScale;
+			}
+		}
+		AnimBankTarget->Modify();
+		AnimBankTarget->PostEditChange();
+#endif // WITH_EDITORONLY_DATA
+	}
+}
+
+void FDataflowAnimBankAssetTerminalNode::Evaluate(UE::Dataflow::FContext& Context) const
+{
+}
+
+bool FDataflowAnimBankAssetTerminalNode::CanAddPin() const 
+{ 
+	return true; 
+}
+
+TArray<UE::Dataflow::FPin> FDataflowAnimBankAssetTerminalNode::AddPins()
+{
+	const int32 Index = AnimSequences.AddDefaulted();
+	const FDataflowInput& Input = RegisterInputArrayConnection(GetConnectionReference(Index));
+	return { { UE::Dataflow::FPin::EDirection::INPUT, Input.GetType(), Input.GetName() } };
+}
+
+bool FDataflowAnimBankAssetTerminalNode::CanRemovePin() const 
+{ 
+	return AnimSequences.Num() > NumInitialInputs;
+}
+
+TArray<UE::Dataflow::FPin> FDataflowAnimBankAssetTerminalNode::GetPinsToRemove() const
+{
+	if (AnimSequences.Num() > 0)
+	{
+		const int32 Index = AnimSequences.Num() - 1;
+		if (const FDataflowInput* const Input = FindInput(GetConnectionReference(Index)))
+		{
+			return { { UE::Dataflow::FPin::EDirection::INPUT, Input->GetType(), Input->GetName() } };
+		}
+	}
+	return Super::GetPinsToRemove();
+}
+
+void FDataflowAnimBankAssetTerminalNode::OnPinRemoved(const UE::Dataflow::FPin& Pin)
+{
+	if (AnimSequences.Num() > 0)
+	{
+		const int32 Index = AnimSequences.Num() - 1;
+		check(AnimSequences.IsValidIndex(Index));
+#if DO_CHECK
+		const FDataflowInput* const Input = FindInput(GetConnectionReference(Index));
+		check(Input);
+		check(Input->GetName() == Pin.Name);
+		check(Input->GetType() == Pin.Type);
+#endif
+		AnimSequences.SetNum(Index);
+	}
+	return Super::OnPinRemoved(Pin);
+}
+
+void FDataflowAnimBankAssetTerminalNode::PostSerialize(const FArchive& Ar)
+{
+	if (Ar.IsLoading())
+	{
+		// Make sure we have the right number of initial inputs
+		for (int32 Index = AnimSequences.Num(); Index < NumInitialInputs; ++Index)
+		{
+			AddPins();
+		}
+		for (int32 Index = 0; Index < NumInitialInputs; ++Index)
+		{
+			check(FindInput(GetConnectionReference(Index)));
+		}
+
+		for (int32 Index = NumInitialInputs; Index < AnimSequences.Num(); ++Index)
+		{
+			FDataflowInput& Input = FindOrRegisterInputArrayConnection(GetConnectionReference(Index));
+		}
+
+		if (Ar.IsTransacting())
+		{
+			const int32 OrigNumRegisteredInputs = GetNumInputs() - NumRequiredDataflowInputs;
+			const int32 OrigNumInputs = AnimSequences.Num();
+			if (OrigNumRegisteredInputs > OrigNumInputs)
+			{
+				// Inputs have been removed.
+				// Temporarily expand Inputs so we can get connection references.
+				AnimSequences.SetNum(OrigNumRegisteredInputs);
+				for (int32 Index = OrigNumInputs; Index < AnimSequences.Num(); ++Index)
+				{
+					UnregisterInputConnection(GetConnectionReference(Index));
+				}
+				AnimSequences.SetNum(OrigNumInputs);
+			}
+		}
+		else
+		{
+			// Index + all Inputs
+			ensureAlways(AnimSequences.Num() + NumRequiredDataflowInputs == GetNumInputs());
+		}
+	}
+}
+
+UE::Dataflow::TConnectionReference<TObjectPtr<UAnimSequence>> FDataflowAnimBankAssetTerminalNode::GetConnectionReference(int32 Index) const
+{
+	return { &AnimSequences[Index], Index, &AnimSequences };
+}
diff --git a/Engine/Plugins/Experimental/ChaosRigidAsset/Source/ChaosRigidAssetNodes/Private/AnimationTerminalNodes.h b/Engine/Plugins/Experimental/ChaosRigidAsset/Source/ChaosRigidAssetNodes/Private/AnimationTerminalNodes.h
new file mode 100644
--- /dev/null
+++ b/Engine/Plugins/Experimental/ChaosRigidAsset/Source/ChaosRigidAssetNodes/Private/AnimationTerminalNodes.h
@@ -0,0 +1,104 @@
+// Copyright Epic Games, Inc. All Rights Reserved.
+
+#pragma once
+
+#include "Dataflow/DataflowTerminalNode.h"
+
+#include "AnimationTerminalNodes.generated.h"
+
+class UAnimSequence;
+class UAnimBank;
+
+////////////////////////////////////////////////////////////////////////////////////////////////////////////////
+
+/**
+ * Terminal node that writes to a animation sequence asset
+ */
+USTRUCT(meta = (DataflowTerminal))
+struct FDataflowAnimSequenceAssetTerminalNode : public FDataflowTerminalNode
+{
+	GENERATED_BODY()
+	DATAFLOW_NODE_DEFINE_INTERNAL(FDataflowAnimSequenceAssetTerminalNode, "AnimSequenceTerminal", "Terminal", "Animation Keys")
+
+public:
+
+	FDataflowAnimSequenceAssetTerminalNode(const UE::Dataflow::FNodeParameters& InParam, FGuid InGuid = FGuid::NewGuid());
+
+private:
+
+	UPROPERTY(meta = (DataflowInput))
+	TObjectPtr<UAnimSequence> Animation;
+
+	UPROPERTY(EditAnywhere, Category = Asset, meta = (DataflowInput))
+	TObjectPtr<UAnimSequence> AnimationAsset;
+
+	void SetAssetValue(TObjectPtr<UObject> Asset, UE::Dataflow::FContext& Context) const override;
+	void Evaluate(UE::Dataflow::FContext& Context) const override;
+};
+
+
+/**
+ * Terminal node that writes to a animation bank asset
+ * an AnimBank is made of multiple anim sequences and used to animate GPU 
+ */
+USTRUCT(meta = (DataflowTerminal))
+struct FDataflowAnimBankAssetTerminalNode : public FDataflowTerminalNode
+{
+	GENERATED_BODY()
+	DATAFLOW_NODE_DEFINE_INTERNAL(FDataflowAnimBankAssetTerminalNode, "AnimBankTerminalNode", "Terminal", "Animation Keys")
+
+public:
+	FDataflowAnimBankAssetTerminalNode(const UE::Dataflow::FNodeParameters& InParam, FGuid InGuid = FGuid::NewGuid());
+
+private:
+	UPROPERTY(EditAnywhere, Category = Asset, meta = (DataflowInput))
+	TObjectPtr<UAnimBank> AnimationBankAsset = nullptr;
+
+	UPROPERTY()
+	TArray<TObjectPtr<UAnimSequence>> AnimSequences;
+
+	UPROPERTY(EditAnywhere, Category = Animation)
+	bool bLooping = true;
+
+	UPROPERTY(EditAnywhere, Category = Animation)
+	bool bAutoStart = true;
+
+	UPROPERTY(EditAnywhere, Category = Animation)
+	float Position = 0.0f;
+
+	UPROPERTY(EditAnywhere, Category = Animation)
+	float PlayRate = 1.0f;
+
+	/**
+	 * Scales the bounds of the instances playing this sequence.
+	 * This is useful when the animation moves the vertices of the mesh outside of its bounds.
+	 * Warning: Increasing the bounds will reduce performance!
+	 */
+	UPROPERTY(EditAnywhere, Category = Animation, meta = (ClampMin = "1", UIMin = "1", ClampMax = "10.0", UIMax = "10.0"))
+	float BoundsScale = 1.0f;
+
+	void SetAssetValue(TObjectPtr<UObject> Asset, UE::Dataflow::FContext& Context) const override;
+	void Evaluate(UE::Dataflow::FContext& Context) const override;
+
+	/* Begin FDataflowNode interface */
+	virtual TArray<UE::Dataflow::FPin> AddPins() override;
+	virtual bool CanAddPin() const override;
+	virtual bool CanRemovePin() const override;
+	virtual TArray<UE::Dataflow::FPin> GetPinsToRemove() const override;
+	virtual void OnPinRemoved(const UE::Dataflow::FPin& Pin) override;
+	virtual void PostSerialize(const FArchive& Ar) override;
+	/* End FDataflowNode interface */
+
+	UE::Dataflow::TConnectionReference<TObjectPtr<UAnimSequence>> GetConnectionReference(int32 Index) const;
+
+	inline static constexpr int32 NumRequiredDataflowInputs = 1;
+	inline static constexpr int32 NumInitialInputs = 2;
+
+};
+
+////////////////////////////////////////////////////////////////////////////////////////////////////////////////
+
+namespace UE::Dataflow
+{
+	void RegisterAnimationTerminalNodes();
+}
diff --git a/Engine/Plugins/Experimental/ChaosRigidAsset/Source/ChaosRigidAssetNodes/Private/ChaosRigidAssetNodesModule.cpp b/Engine/Plugins/Experimental/ChaosRigidAsset/Source/ChaosRigidAssetNodes/Private/ChaosRigidAssetNodesModule.cpp
--- a/Engine/Plugins/Experimental/ChaosRigidAsset/Source/ChaosRigidAssetNodes/Private/ChaosRigidAssetNodesModule.cpp
+++ b/Engine/Plugins/Experimental/ChaosRigidAsset/Source/ChaosRigidAssetNodes/Private/ChaosRigidAssetNodesModule.cpp
@@ -3,15 +3,19 @@
 #include "ChaosRigidAssetNodesModule.h"
 #include "ShapeElemNodes.h"
 #include "PhysicsAssetDataflowNodes.h"
+#include "AnimationTerminalNodes.h"
+#include "Generators/SimulationAffector.h"
 
 IMPLEMENT_MODULE(FChaosRigidAssetNodesModule, ChaosRigidAssetNodes);
 
  void FChaosRigidAssetNodesModule::StartupModule()
 {
+	UE::Dataflow::RegisterSimulationAffectorNodes();
 	UE::Dataflow::RegisterPhysicsAssetTerminalNode();
 	UE::Dataflow::RegisterPhysicsAssetNodes();
 	UE::Dataflow::RegisterShapeNodes();
-}
+	UE::Dataflow::RegisterAnimationTerminalNodes();
+ }
 
  void FChaosRigidAssetNodesModule::ShutdownModule()
 {
diff --git a/Engine/Plugins/Experimental/ChaosRigidAsset/Source/ChaosRigidAssetNodes/Private/Generators/SimulationAffector.cpp b/Engine/Plugins/Experimental/ChaosRigidAsset/Source/ChaosRigidAssetNodes/Private/Generators/SimulationAffector.cpp
new file mode 100644
--- /dev/null
+++ b/Engine/Plugins/Experimental/ChaosRigidAsset/Source/ChaosRigidAssetNodes/Private/Generators/SimulationAffector.cpp
@@ -0,0 +1,91 @@
+// Copyright Epic Games, Inc. All Rights Reserved.
+
+#include "Generators/SimulationAffector.h"
+
+#include "Chaos/ParticleHandle.h"
+#include "Dataflow/DataflowNodeFactory.h"
+#include "Dataflow/DataflowAnyTypeRegistry.h"
+#include "Math/UnitConversion.h"
+#include "Physics/ImmediatePhysics/ImmediatePhysicsChaos/ImmediatePhysicsActorHandle_Chaos.h"
+
+
+namespace UE::Dataflow
+{
+	void RegisterSimulationAffectorNodes()
+	{
+		UE_DATAFLOW_REGISTER_ANYTYPE(FDataflowSimulationAffectorTypes);
+
+		DATAFLOW_NODE_REGISTER_CREATION_FACTORY(FDataflowSimulationAccelerationAffectorNode);
+		DATAFLOW_NODE_REGISTER_CREATION_FACTORY(FDataflowSimulationWindAffectorNode);
+	}
+}
+
+////////////////////////////////////////////////////////////////////////////////////////////////////////////////
+
+void FDataflowSimulationAccelerationAffector::Process(const TArray<ImmediatePhysics::FActorHandle*>& Actors) const
+{
+	for (ImmediatePhysics::FActorHandle* Actor : Actors)
+	{
+		if (Actor && Actor->IsSimulated())
+		{
+			Actor->AddForce(Acceleration, ImmediatePhysics::EForceType::AddAcceleration);
+		}
+	}
+}
+
+FDataflowSimulationAccelerationAffectorNode::FDataflowSimulationAccelerationAffectorNode(const UE::Dataflow::FNodeParameters& InParam, FGuid InGuid)
+	: FDataflowNode(InParam, InGuid)
+{
+	RegisterOutputConnection(&Affector);
+}
+
+void FDataflowSimulationAccelerationAffectorNode::Evaluate(UE::Dataflow::FContext& Context, const FDataflowOutput* Out) const
+{
+	check(Out->IsA(&Affector));
+	SetValue(Context, Affector, &Affector);
+}
+
+////////////////////////////////////////////////////////////////////////////////////////////////////////////////
+
+void FDataflowSimulationWindAffector::Process(const TArray<ImmediatePhysics::FActorHandle*>& Actors) const
+{
+	const FVector NormalizedWindDirection = WindDirection.GetSafeNormal(0, FVector::XAxisVector);
+	const float WindSpeedInCmPerSecond = FUnitConversion::Convert(WindSpeed, EUnit::KilometersPerHour, EUnit::CentimetersPerSecond);
+	const float AirDensityInKgPerCm3 = FUnitConversion::Convert(AirDensity, EUnit::KilogramsPerCubicMeter, EUnit::KilogramsPerCubicCentimeter);
+	const FVector WindVelocitySquared = NormalizedWindDirection * ( WindSpeedInCmPerSecond * WindSpeedInCmPerSecond);
+	const FVector WindPressure = (WindVelocitySquared * AirDensityInKgPerCm3 * 0.5f);
+
+	for (ImmediatePhysics::FActorHandle* Actor : Actors)
+	{
+		if (Actor && Actor->IsSimulated())
+		{
+			if (const Chaos::FGeometryParticleHandle* Particle = Actor->GetParticle())
+			{
+				const FTransform ActorTransform = Actor->GetWorldTransform();
+				// gross estimation the wind facing area 
+				const float TotalVolumeInCm3 = Particle->LocalBounds().GetVolume();
+				const float CubeSizeInCm = FMath::Pow(TotalVolumeInCm3, 1.0 / 3.0);
+				const float EstimatedAreaInCm2 = CubeSizeInCm * CubeSizeInCm;
+
+				// then estimate how much we are facing the wind
+				const float DotProductWithZAxis = ActorTransform.GetUnitAxis(EAxis::Z).Dot(NormalizedWindDirection);
+				const float FacingFactor = FMath::Clamp(1.0 - FMath::Abs(DotProductWithZAxis), 0.0, 1.0);
+
+				const FVector WindForce = (EstimatedAreaInCm2 * FacingFactor) * WindPressure;
+				Actor->AddForce(WindForce);
+			}
+		}
+	}
+}
+
+FDataflowSimulationWindAffectorNode::FDataflowSimulationWindAffectorNode(const UE::Dataflow::FNodeParameters& InParam, FGuid InGuid)
+	: FDataflowNode(InParam, InGuid)
+{
+	RegisterOutputConnection(&Affector);
+}
+
+void FDataflowSimulationWindAffectorNode::Evaluate(UE::Dataflow::FContext& Context, const FDataflowOutput* Out) const
+{
+	check(Out->IsA(&Affector));
+	SetValue(Context, Affector, &Affector);
+}
\ No newline at end of file
diff --git a/Engine/Plugins/Experimental/ChaosRigidAsset/Source/ChaosRigidAssetNodes/Private/Generators/SimulationAffector.h b/Engine/Plugins/Experimental/ChaosRigidAsset/Source/ChaosRigidAssetNodes/Private/Generators/SimulationAffector.h
new file mode 100644
--- /dev/null
+++ b/Engine/Plugins/Experimental/ChaosRigidAsset/Source/ChaosRigidAssetNodes/Private/Generators/SimulationAffector.h
@@ -0,0 +1,106 @@
+// Copyright Epic Games, Inc. All Rights Reserved.
+
+#pragma once
+
+#include "Dataflow/DataflowNode.h"
+#include "Dataflow/DataflowTypePolicy.h"
+#include "Physics/ImmediatePhysics/ImmediatePhysicsDeclares.h"
+
+#include "SimulationAffector.generated.h"
+
+USTRUCT()
+struct FDataflowSimulationAffector
+{
+	GENERATED_BODY();
+public:
+	virtual ~FDataflowSimulationAffector() = default;
+	virtual void Process(const TArray<ImmediatePhysics::FActorHandle*>& Actors) const {};
+};
+
+/** UStruct based types */
+USTRUCT()
+struct FDataflowSimulationAffectorTypes : public FDataflowAnyType
+{
+	using FPolicyType = FDataflowUDerivedStructTypePolicy<FDataflowSimulationAffector>;
+	using FStorageType = void;
+	GENERATED_BODY()
+};
+
+namespace UE::Dataflow
+{
+	void RegisterSimulationAffectorNodes();
+}
+
+////////////////////////////////////////////////////////////////////////////////////////////////////////////////
+
+USTRUCT()
+struct FDataflowSimulationAccelerationAffector: public FDataflowSimulationAffector
+{
+	GENERATED_BODY();
+
+	virtual ~FDataflowSimulationAccelerationAffector() override = default;
+	virtual void Process(const TArray<ImmediatePhysics::FActorHandle*>& Actors) const override;
+
+private:
+	UPROPERTY(EditAnywhere, Category = Physics, meta = (ForceUnits = "cm/s2"))
+	FVector Acceleration = FVector(0, 0, -980); // default to earth gravity
+};
+
+/**
+ * Output a affector that apply a constant acceleration to all dynamic actors
+ */
+USTRUCT()
+struct FDataflowSimulationAccelerationAffectorNode : public FDataflowNode
+{
+	GENERATED_BODY()
+	DATAFLOW_NODE_DEFINE_INTERNAL(FDataflowSimulationAccelerationAffectorNode, "AccelerationAffector", "PhysicsAsset", "")
+
+public:
+	FDataflowSimulationAccelerationAffectorNode(const UE::Dataflow::FNodeParameters& InParam, FGuid InGuid = FGuid::NewGuid());
+
+private:
+	UPROPERTY(EditAnywhere, Category = Affector, meta = (DataflowOutput, ShowOnlyInnerProperties))
+	FDataflowSimulationAccelerationAffector Affector;
+
+	void Evaluate(UE::Dataflow::FContext& Context, const FDataflowOutput* Out) const override;
+};
+
+////////////////////////////////////////////////////////////////////////////////////////////////////////////////
+
+USTRUCT()
+struct FDataflowSimulationWindAffector : public FDataflowSimulationAffector
+{
+	GENERATED_BODY();
+
+	virtual ~FDataflowSimulationWindAffector() override = default;
+	virtual void Process(const TArray<ImmediatePhysics::FActorHandle*>& Actors) const override;
+
+private:
+	UPROPERTY(EditAnywhere, Category = Physics, meta = (ForceUnits = "kg/m3"))
+	float AirDensity = 1.2f; // standard air density at about 20C
+
+	UPROPERTY(EditAnywhere, Category = Physics, meta = (ForceUnits = "km/h"))
+	float WindSpeed = 15.f;
+
+	UPROPERTY(EditAnywhere, Category = Physics)
+	FVector WindDirection = FVector::XAxisVector;
+};
+
+/**
+ * Output a affector that apply a Wind force onto the rigid bodies
+ */
+USTRUCT()
+struct FDataflowSimulationWindAffectorNode : public FDataflowNode
+{
+	GENERATED_BODY()
+	DATAFLOW_NODE_DEFINE_INTERNAL(FDataflowSimulationWindAffectorNode, "WindAffector", "PhysicsAsset", "")
+
+public:
+	FDataflowSimulationWindAffectorNode(const UE::Dataflow::FNodeParameters& InParam, FGuid InGuid = FGuid::NewGuid());
+
+private:
+	UPROPERTY(EditAnywhere, Category = Affector, meta = (DataflowOutput, ShowOnlyInnerProperties))
+	FDataflowSimulationWindAffector Affector;
+
+	void Evaluate(UE::Dataflow::FContext& Context, const FDataflowOutput* Out) const override;
+};
\ No newline at end of file
diff --git a/Engine/Plugins/Experimental/ChaosRigidAsset/Source/ChaosRigidAssetNodes/Private/PhysicsAssetDataflowNodes.cpp b/Engine/Plugins/Experimental/ChaosRigidAsset/Source/ChaosRigidAssetNodes/Private/PhysicsAssetDataflowNodes.cpp
--- a/Engine/Plugins/Experimental/ChaosRigidAsset/Source/ChaosRigidAssetNodes/Private/PhysicsAssetDataflowNodes.cpp
+++ b/Engine/Plugins/Experimental/ChaosRigidAsset/Source/ChaosRigidAssetNodes/Private/PhysicsAssetDataflowNodes.cpp
@@ -56,7 +56,6 @@ void FDataflowPhysicsAssetTerminalNode::Evaluate(UE::Dataflow::FContext& Context
 void UE::Dataflow::RegisterPhysicsAssetTerminalNode()
 {
 	DATAFLOW_NODE_REGISTER_CREATION_FACTORY(FDataflowPhysicsAssetTerminalNode);
-	DATAFLOW_NODE_REGISTER_CREATION_FACTORY(FDataflowAnimSequenceAssetTerminalNode);
 }
 
 void UE::Dataflow::RegisterPhysicsAssetNodes()
@@ -846,7 +845,7 @@ void FDataflowAggGeomAddShape::Evaluate(UE::Dataflow::FContext& Context, const F
 
 void FDataflowRigidSimulationCachingNode::Register(const UE::Dataflow::FNodeParameters& InParam)
 {
-	AddInputs(State, Duration, FrameRate, Gravity, Velocity);
+	AddInputs(State, Duration, FrameRate, Gravity, Velocity, Affector);
 	AddOutputs(Animation);
 }
 
@@ -866,6 +865,8 @@ void FDataflowRigidSimulationCachingNode::Evaluate(UE::Dataflow::FContext& Conte
 		const FVector InGravity= GetValue(Context, &Gravity);
 		const FVector InVelocity = GetValue(Context, &Velocity);
 		
+		const FConstStructView AffectorStructView = GetValue(Context, &Affector).GetConstStructView();
+		const FDataflowSimulationAffector* InAffector = AffectorStructView.GetPtr<FDataflowSimulationAffector>();
 
 		const float FixedTimeStep = 1.0f / static_cast<float>(InFrameRate);
 		const int32 NumFrames = FMath::Max(1, (int32)(InDuration * InFrameRate));
@@ -947,12 +948,15 @@ void FDataflowRigidSimulationCachingNode::Evaluate(UE::Dataflow::FContext& Conte
 								RawTrack->RotKeys.Add(FQuat4f(ActorLocalTransform.GetRotation()));
 								RawTrack->ScaleKeys.Add(FVector3f(ActorLocalTransform.GetScale3D()));
 							}
-
-							const float Ratio = FMath::Clamp(Handle->GetParticle()->GetX().Z / 2800.f, 0, 1);
-							Handle->SetLinearVelocity(InVelocity * Ratio * Ratio);
 						}
 					}
 
+					// apply affectors
+					if (InAffector)
+					{
+						InAffector->Process(SimState.Actors);
+					}
+
 					SimState.Simulation->UpdateSimulationSpace(FTransform::Identity, InVelocity, FVector::ZeroVector, FVector::ZeroVector, FVector::ZeroVector);
 					RunSimulation(SimState, InGravity);
 				}
@@ -1006,7 +1010,6 @@ void FDataflowRigidSimulationCachingNode::BuildSimulationFromState(const FPhysic
 	CachedTransforms.SetNumUninitialized(RefSkeleton.GetRefBonePose().Num());
 	CachedTransformReady.SetNumZeroed(RefSkeleton.GetRefBonePose().Num());
 
-
 	for (TObjectPtr<USkeletalBodySetup> BodySetup : InState.Bodies)
 	{
 		TUniquePtr<FBodyInstance> BodyInstance = CreateBodyInstance(BodySetup);
@@ -1031,6 +1034,7 @@ void FDataflowRigidSimulationCachingNode::BuildSimulationFromState(const FPhysic
 				NewBodyHandle->SetHasCollision(false);
 				//OutSimState.Simulation->AddToCollidingPairs(NewBodyHandle);
 				OutSimState.ActorsByName.Add(BodySetup->BoneName, NewBodyHandle);
+				OutSimState.Actors.Add(NewBodyHandle);
 				UE_LOG(LogTemp, Warning, TEXT("@@ Create Sim Actor : [%s] - bSimulated[%d] - Transform[%s]"), *BodySetup->BoneName.ToString(), (int)bSimulated, *BoneComponentSpaceTransform.ToString());
 			}
 		}
@@ -1088,42 +1092,6 @@ TUniquePtr<FBodyInstance> FDataflowRigidSimulationCachingNode::CreateBodyInstanc
 	return BodyInstance;
 }
 
-////////////////////////////////////////////////////////////////////////////////////////////////
-FDataflowAnimSequenceAssetTerminalNode::FDataflowAnimSequenceAssetTerminalNode(const UE::Dataflow::FNodeParameters& InParam, FGuid InGuid)
-	: FDataflowTerminalNode(InParam, InGuid)
-{
-	RegisterInputConnection(&Animation);
-	RegisterInputConnection(&AnimationAsset);
-}
-
-void FDataflowAnimSequenceAssetTerminalNode::SetAssetValue(TObjectPtr<UObject> Asset, UE::Dataflow::FContext& Context) const
-{
-	UAnimSequence* AnimTarget = Asset ? Cast<UAnimSequence>(Asset.Get()) : nullptr;
-	if (!AnimTarget)
-	{
-		AnimTarget = GetValue(Context, &AnimationAsset);
-	}
-
-	if (AnimTarget)
-	{
-#if WITH_EDITOR
-		// This should be const but CreateAnimation does not take a const parameter
-		if (UAnimSequence* InAnimation = GetValue(Context, &Animation))
-		{
-			// copied from UAnimSequence::CreateAnimation because the method does not take a const anim sequence as a parameter
-			AnimTarget->CreateAnimation(InAnimation);
-		}
-
-		AnimTarget->Modify();
-		AnimTarget->PostEditChange();
-#endif // WITH_EDITOR
-	}
-}
-
-void FDataflowAnimSequenceAssetTerminalNode::Evaluate(UE::Dataflow::FContext& Context) const
-{
-}
-
 ////////////////////////////////////////////////////////////////////////////////////////////////
 
 void FDataflowAggGeomAddShape::Register()
diff --git a/Engine/Plugins/Experimental/ChaosRigidAsset/Source/ChaosRigidAssetNodes/Private/PhysicsAssetDataflowNodes.h b/Engine/Plugins/Experimental/ChaosRigidAsset/Source/ChaosRigidAssetNodes/Private/PhysicsAssetDataflowNodes.h
--- a/Engine/Plugins/Experimental/ChaosRigidAsset/Source/ChaosRigidAssetNodes/Private/PhysicsAssetDataflowNodes.h
+++ b/Engine/Plugins/Experimental/ChaosRigidAsset/Source/ChaosRigidAssetNodes/Private/PhysicsAssetDataflowNodes.h
@@ -12,6 +12,7 @@
 #include "PhysicsAssetDataflowState.h"
 #include "Generators/BoneGeometryGenerators.h"
 #include "Generators/ConstraintGenerators.h"
+#include "Generators/SimulationAffector.h"
 #include "BoneSelection.h"
 #include "RigidDataflowNode.h"
 #include "ShapeElemNodes.h"
@@ -383,6 +384,7 @@ private:
 	UPROPERTY(meta=(DataflowInput))
 	TObjectPtr<UBoneGeometryGenerator> Generator;
 
+
 	UPROPERTY(meta=(DataflowInput))
 	TObjectPtr<USkeletalBodySetup> TemplateBody;
 
@@ -569,6 +571,9 @@ private:
 	UPROPERTY(EditAnywhere, Category = Simulation, meta = (DataflowInput))
 	FVector Velocity = FVector(0, 0, 0);
 
+	UPROPERTY(EditAnywhere, Category = Simulation, meta = (DataflowInput))
+	FDataflowSimulationAffectorTypes Affector;
+
 	void Evaluate(UE::Dataflow::FContext& Context, const FDataflowOutput* Out) const override;
 	void Register(const UE::Dataflow::FNodeParameters& InParam);
 
@@ -579,38 +584,13 @@ private:
 		float FixedTimeStep = 0;
 		TUniquePtr<ImmediatePhysics::FSimulation> Simulation;
 		TMap<FName, ImmediatePhysics::FActorHandle*> ActorsByName;
+		TArray<ImmediatePhysics::FActorHandle*> Actors;
 	};
 
 	void BuildSimulationFromState(const FPhysicsAssetDataflowState& InState, FSimulationState& OutSimState) const;
 	void RunSimulation(FSimulationState& Simulation, const FVector& InGravity) const;
 };
 
-
-/**
- * Terminal node that writes to a animation sequence asset
- */
-USTRUCT(meta = (DataflowTerminal))
-struct FDataflowAnimSequenceAssetTerminalNode : public FDataflowTerminalNode
-{
-	GENERATED_BODY()
-	DATAFLOW_NODE_DEFINE_INTERNAL(FDataflowAnimSequenceAssetTerminalNode, "AnimSequenceTerminal", "Terminal", "")
-
-public:
-
-	FDataflowAnimSequenceAssetTerminalNode(const UE::Dataflow::FNodeParameters& InParam, FGuid InGuid = FGuid::NewGuid());
-
-private:
-
-	UPROPERTY(meta = (DataflowInput))
-	TObjectPtr<UAnimSequence> Animation;
-
-	UPROPERTY(EditAnywhere, Category = Asset, meta = (DataflowInput))
-	TObjectPtr<UAnimSequence> AnimationAsset;
-
-	void SetAssetValue(TObjectPtr<UObject> Asset, UE::Dataflow::FContext& Context) const override;
-	void Evaluate(UE::Dataflow::FContext& Context) const override;
-};
-
 ////////////////////////////////////////////////////////////////////////////////////////////////////////////////
 
 namespace UE::Dataflow
