[Dataflow] [Physics Asset] Added ability to merge bones within selections with similar functionality to the physics asset editor where we inspect the bone geometry and make determinations on whether a bone can be merged or skipped entirely.

Bone Selection merging is generically applied and results in a new selection and is not limited to use in just the geometry generator or constraint generator use cases.

#Synced-CL 48950948

diff --git a/Engine/Plugins/Experimental/ChaosRigidAsset/Source/ChaosRigidAssetEditor/Private/DataflowRendering.cpp b/Engine/Plugins/Experimental/ChaosRigidAsset/Source/ChaosRigidAssetEditor/Private/DataflowRendering.cpp
--- a/Engine/Plugins/Experimental/ChaosRigidAsset/Source/ChaosRigidAssetEditor/Private/DataflowRendering.cpp
+++ b/Engine/Plugins/Experimental/ChaosRigidAsset/Source/ChaosRigidAssetEditor/Private/DataflowRendering.cpp
@@ -257,7 +257,16 @@ namespace UE::Chaos::RigidAsset
 
 			TArray<FTransform> BoneTransforms;
 			TArray<int32> ParentIndices;
-			TArray<bool> SelectedIndices;
+
+
+			enum class ELocalSelectionType
+			{
+				None = 0,
+				Selected,
+				Merged
+			};
+
+			TArray<ELocalSelectionType> SelectedIndices;
 
 			RefSkel.GetBoneAbsoluteTransforms(BoneTransforms);
 
@@ -266,7 +275,12 @@ namespace UE::Chaos::RigidAsset
 			SelectedIndices.SetNumZeroed(NumBones);
 			for(const FRigidAssetBoneInfo& Selected : Selection.SelectedBones)
 			{
-				SelectedIndices[Selected.Index] = true;
+				SelectedIndices[Selected.Index] = ELocalSelectionType::Selected;
+
+				for(int32 MergeIndex : Selected.MergedChildren)
+				{
+					SelectedIndices[MergeIndex] = ELocalSelectionType::Merged;
+				}
 			}
 
 			ParentIndices.SetNum(NumBones);
@@ -285,7 +299,20 @@ namespace UE::Chaos::RigidAsset
 				const int32 ParentIndex = ParentIndices[BoneIndex];
 				if(ParentIndex != INDEX_NONE)
 				{
-					const FLinearColor BoneColor = SelectedIndices[ParentIndex] ? FLinearColor(1.0f, 0.6f, 0.0f) : FLinearColor::Gray;
+					FLinearColor BoneColor = FLinearColor::Gray;
+
+					switch(SelectedIndices[ParentIndex])
+					{
+					case ELocalSelectionType::Selected:
+						BoneColor = FLinearColor(1.0f, 0.6f, 0.0f);
+						break;
+					case ELocalSelectionType::Merged:
+						BoneColor = FLinearColor(0.0f, 0.8f, 0.0f);
+						break;
+					default:
+						break;
+					}
+
 					DrawBone(RenderingFacade, (FVector3f)BoneTransforms[ParentIndex].GetTranslation(), (FVector3f)BoneTransforms[BoneIndex].GetTranslation(), BoneColor);
 				}
 			}
diff --git a/Engine/Plugins/Experimental/ChaosRigidAsset/Source/ChaosRigidAssetNodes/Private/Generators/BoneGeometryGenerators.cpp b/Engine/Plugins/Experimental/ChaosRigidAsset/Source/ChaosRigidAssetNodes/Private/Generators/BoneGeometryGenerators.cpp
--- a/Engine/Plugins/Experimental/ChaosRigidAsset/Source/ChaosRigidAssetNodes/Private/Generators/BoneGeometryGenerators.cpp
+++ b/Engine/Plugins/Experimental/ChaosRigidAsset/Source/ChaosRigidAssetNodes/Private/Generators/BoneGeometryGenerators.cpp
@@ -121,17 +121,7 @@ void UBoneGeometryGenerator::DebugDraw(UE::Dataflow::FContext& Context, IDataflo
 
 ////////////////////////////////////////////////////////////////////////////////////////////////////////////////
 
-/**
- * Data extracted from an LOD for body generation
- */
-struct FBoneVertData
-{
-	TArray<FVector3f> Positions;
-	TArray<FVector3f> Normals;
-	TArray<float> Weights;
-};
-
-FBoneVertData GetBoneVertexData(FName BoneName, USkeletalMesh* Mesh, int32 Lod, EVertexSelectMode Mode)
+FBoneVertData GetBoneVertexData(FName BoneName, USkeletalMesh* Mesh, int32 Lod, EVertexSelectMode SelectMode)
 {
 	if(!Mesh)
 	{
@@ -167,6 +157,7 @@ FBoneVertData GetBoneVertexData(FName BoneName, USkeletalMesh* Mesh, int32 Lod,
 		return {};
 	}
 
+	FBoxSphereBounds3f::Builder BoundsBuilder;
 	const int32 NumSections = LodData.RenderSections.Num();
 
 	for(int32 SectionIndex = 0; SectionIndex < NumSections; ++SectionIndex)
@@ -176,7 +167,7 @@ FBoneVertData GetBoneVertexData(FName BoneName, USkeletalMesh* Mesh, int32 Lod,
 
 		for(uint32 VertIndex = Section.BaseVertexIndex; VertIndex < Section.BaseVertexIndex + Section.NumVertices; ++VertIndex)
 		{
-			if(Mode == EVertexSelectMode::Any)
+			if(SelectMode == EVertexSelectMode::Any)
 			{
 				for(uint32 InfluenceIndex = 0; InfluenceIndex < SectionMaxInfluences; InfluenceIndex++)
 				{
@@ -192,6 +183,7 @@ FBoneVertData GetBoneVertexData(FName BoneName, USkeletalMesh* Mesh, int32 Lod,
 						Result.Positions.Add(RefMatrix.TransformPosition(Positions.VertexPosition(VertIndex)));
 						Result.Normals.Add(RefMatrix.TransformVector(MeshVertexBuffer.VertexTangentZ(VertIndex)));
 						Result.Weights.Add(BoneWeight * WeightsToFloatMult);
+						BoundsBuilder += Result.Positions.Last();
 
 						// Vertex will only have one influence for a given bone, skip the rest
 						break;
@@ -223,24 +215,200 @@ FBoneVertData GetBoneVertexData(FName BoneName, USkeletalMesh* Mesh, int32 Lod,
 					Result.Positions.Add(RefMatrix.TransformPosition(Positions.VertexPosition(VertIndex)));
 					Result.Normals.Add(RefMatrix.TransformVector(MeshVertexBuffer.VertexTangentZ(VertIndex)));
 					Result.Weights.Add(MaxInfluence * WeightsToFloatMult);
+					BoundsBuilder += Result.Positions.Last();
+				}
+			}
+		}
+	}
+
+	Result.Bounds = BoundsBuilder;
+	return Result;
+}
+
+FBoneVertData GetMergedBoneVertexData(const FRigidAssetBoneInfo& Bone, USkeletalMesh* Mesh, int32 Lod, EVertexSelectMode SelectMode)
+{
+	if(!Mesh)
+	{
+		return {};
+	}
+
+	FBoneVertData RootData = GetBoneVertexData(Bone.Name, Mesh, Lod, SelectMode);
+
+	const FReferenceSkeleton& RefSkeleton = Mesh->GetRefSkeleton();
+	const TArray<FTransform>& RefTransforms = RefSkeleton.GetRawRefBonePose();
+
+	for(int32 BoneIndex : Bone.MergedChildren)
+	{
+		FBoneVertData ChildData = GetBoneVertexData(RefSkeleton.GetBoneName(BoneIndex), Mesh, Lod, SelectMode);
+		FTransform ToParent = RefTransforms[BoneIndex];
+
+		// Merged bones may be many bones away, need to walk up the chain and combine the transforms
+		int32 ParentIndex = RefSkeleton.GetParentIndex(BoneIndex);
+		while(ParentIndex != Bone.Index && ParentIndex != INDEX_NONE)
+		{
+			ToParent *= RefTransforms[ParentIndex];
+			ParentIndex = RefSkeleton.GetParentIndex(ParentIndex);
+		}
+
+		if(ParentIndex == INDEX_NONE)
+		{
+			// Somehow have a connection to root, with a correctly formed selection this should not be possible
+			return {};
+		}
+
+		RootData.Positions.Reserve(RootData.Positions.Num() + ChildData.Positions.Num());
+		RootData.Normals.Reserve(RootData.Normals.Num() + ChildData.Normals.Num());
+		for(int32 VertIndex = 0; VertIndex < ChildData.Positions.Num(); ++VertIndex)
+		{
+			const FVector PositionInParentFrame = ToParent.TransformPosition(FVector(ChildData.Positions[VertIndex]));
+			const FVector NormalInParentFrame = ToParent.TransformVector(FVector(ChildData.Normals[VertIndex]));
+
+			RootData.Positions.Add(FVector3f(PositionInParentFrame));
+			RootData.Normals.Add(FVector3f(NormalInParentFrame));
+		}
+
+		RootData.Weights.Append(ChildData.Weights);
+		RootData.Bounds = RootData.Bounds + ChildData.Bounds.TransformBy(FTransform3f(ToParent));
+	}
+
+	return RootData;
+}
+
+FRigidAssetBoneSelection MergeSelection(const FBaseGenerationSettings& Settings, const FRigidAssetBoneSelection& InSelectionSorted)
+{
+	if(!InSelectionSorted.Mesh)
+	{
+		return {};
+	}
+
+	FRigidAssetBoneSelection Result;
+
+	struct FLocalBoneData
+	{
+		FRigidAssetBoneInfo Bone;
+		FRigidAssetBoneInfo RootBone;
+
+		FBoneVertData VertData;
+
+		bool bRootBone = false;
+
+		TArray<FRigidAssetBoneInfo> ChildBones;
+	};
+
+	TObjectPtr<USkeletalMesh> Mesh = InSelectionSorted.Mesh;
+	const FReferenceSkeleton& RefSkel = Mesh->GetRefSkeleton();
+	const TArray<FTransform>& RefTransforms = RefSkel.GetRawRefBonePose();
+
+	const int32 NumSelectionBones = InSelectionSorted.SelectedBones.Num();
+	TArray<FLocalBoneData> BoneData;
+	BoneData.Reserve(NumSelectionBones);
+
+	for(const FRigidAssetBoneInfo& Bone : InSelectionSorted.SelectedBones)
+	{
+		BoneData.AddDefaulted();
+		FLocalBoneData& Local = BoneData.Last();
+
+		Local.Bone = Bone;
+		Local.RootBone = Bone;
+		Local.VertData = GetBoneVertexData(Bone.Name, Mesh, Settings.SourceLod, Settings.VertexMode);
+	}
+
+	const float MinSizeSq = Settings.Operation == EMergeOperation::MergeSmall ? Settings.MinimumBoneSize * Settings.MinimumBoneSize : std::numeric_limits<float>::max();
+	const float MergeThresholdSizeSq = Settings.Operation == EMergeOperation::MergeSmall ? Settings.SmallBoneMergeThresholdOverride * Settings.SmallBoneMergeThresholdOverride : 0;
+
+	TMap<int32, int32> BoneIndexToMergedIndex;
+
+	// Walk backwards, merging child bones into parent
+	const int32 NumBones = InSelectionSorted.SelectedBones.Num();
+	for(int32 BoneIndex = NumBones - 1; BoneIndex >= 0; --BoneIndex)
+	{
+		FLocalBoneData& LocalData = BoneData[BoneIndex];
+		const FRigidAssetBoneInfo& Bone = LocalData.Bone;
+
+		const float BoneSizeSq = LocalData.VertData.Bounds.BoxExtent.SizeSquared() * 2;
+
+		if(BoneSizeSq < MinSizeSq)
+		{
+			if(BoneSizeSq < MergeThresholdSizeSq || (Settings.SmallBoneOp == ESmallBoneOperation::Skip && Settings.Operation == EMergeOperation::MergeSmall) || LocalData.VertData.Positions.Num() == 0)
+			{
+				// Skip this bone
+				LocalData.RootBone = {};
+				continue;
+			}
+
+			const int32 ParentSkelIndex = RefSkel.GetParentIndex(Bone.Index);
+			const int32 ParentBoneIndex = InSelectionSorted.SelectedBones.IndexOfByPredicate([ParentSkelIndex](const FRigidAssetBoneInfo& SearchBone)
+				{
+					return SearchBone.Index == ParentSkelIndex;
+				});
+
+			if(ParentBoneIndex != INDEX_NONE)
+			{
+				FLocalBoneData& ParentLocalData = BoneData[ParentBoneIndex];
+				FBoneVertData& ParentVertData = ParentLocalData.VertData;
+
+				// Need to transform all child verts into parent space
+				FTransform ToParent = RefTransforms[Bone.Index];
+
+				for(int32 ChildVertIndex = 0; ChildVertIndex < LocalData.VertData.Positions.Num(); ++ChildVertIndex)
+				{
+					FVector PositionInParentFrame = ToParent.TransformPosition(FVector(LocalData.VertData.Positions[ChildVertIndex]));
+					FVector NormalInParentFrame = ToParent.TransformVector(FVector(LocalData.VertData.Normals[ChildVertIndex]));
+
+					ParentVertData.Positions.Add(FVector3f(PositionInParentFrame));
+					ParentVertData.Normals.Add(FVector3f(NormalInParentFrame));
 				}
+
+				// Weights don't need transforming - just append them
+				ParentVertData.Weights.Append(LocalData.VertData.Weights);
+				ParentVertData.Bounds = ParentVertData.Bounds + LocalData.VertData.Bounds.TransformBy(FTransform3f(ToParent));
+
+				ParentLocalData.ChildBones.Add(LocalData.RootBone);
+				ParentLocalData.ChildBones.Append(LocalData.ChildBones);
+
+				// Clear this bone so we can filter easily after the merge is complete
+				LocalData.RootBone = {};
 			}
 		}
 	}
 
+	auto IsMergedBone = [](const FLocalBoneData& Bone)
+		{
+			return Bone.RootBone.Index == INDEX_NONE;
+		};
+
+	BoneData.RemoveAllSwap(IsMergedBone, EAllowShrinking::No);
+
+	// Convert local bone data back to a selection
+	Result.Mesh = InSelectionSorted.Mesh;
+	Result.Skeleton = InSelectionSorted.Skeleton;
+
+	for(const FLocalBoneData& LocalBone : BoneData)
+	{
+		Result.SelectedBones.Emplace(LocalBone.Bone);
+		FRigidAssetBoneInfo& NewSelectionBone = Result.SelectedBones.Last();
+
+		for(const FRigidAssetBoneInfo& ChildBone : LocalBone.ChildBones)
+		{
+			NewSelectionBone.MergedChildren.Add(ChildBone.Index);
+		}
+	}
+
 	return Result;
 }
 
 ////////////////////////////////////////////////////////////////////////////////////////////////////////////////
 
-TArray<UE::Chaos::RigidAsset::FSimpleGeometry> UBoneGeometryGenerator_Box::Build(FRigidAssetBoneSelection Bones)
+TArray<FRigidAssetBoneGeometry> UBoneGeometryGenerator_Box::Build(FRigidAssetBoneSelection Bones)
 {
+	using namespace UE::Chaos::RigidAsset;
+
 	if(!Bones.Mesh || !Bones.Skeleton)
 	{
 		return {};
 	}
 
-	TArray<UE::Chaos::RigidAsset::FSimpleGeometry> Result;
+	TArray<FRigidAssetBoneGeometry> Result;
 
 	for(const FRigidAssetBoneInfo& Bone : Bones.SelectedBones)
 	{
@@ -250,8 +418,8 @@ TArray<UE::Chaos::RigidAsset::FSimpleGeometry> UBoneGeometryGenerator_Box::Build
 		{
 			// No influences, make a small sphere and warn
 			// #TODO pass context through for dataflow warnings
-			Result.AddDefaulted();
-			Result.Last().GeometryElement = MakeShared<FKSphereElem>(5.0f);
+			Result.Emplace(Bone, FSimpleGeometry(MakeShared<FKSphereElem>(5.0f)));
+
 			continue;
 		}
 
@@ -263,8 +431,7 @@ TArray<UE::Chaos::RigidAsset::FSimpleGeometry> UBoneGeometryGenerator_Box::Build
 		NewBox->Z = VertBox.GetExtent().Z * 2.0f;
 		NewBox->Center = static_cast<FVector>(VertBox.GetCenter());
 
-		Result.AddDefaulted();
-		Result.Last().GeometryElement = NewBox;
+		Result.Emplace(Bone, FSimpleGeometry(NewBox));
 	}
 
 	return Result;
@@ -272,36 +439,36 @@ TArray<UE::Chaos::RigidAsset::FSimpleGeometry> UBoneGeometryGenerator_Box::Build
 
 ////////////////////////////////////////////////////////////////////////////////////////////////////////////////
 
-TArray<UE::Chaos::RigidAsset::FSimpleGeometry> UBoneGeometryGenerator_Sphere::Build(FRigidAssetBoneSelection Bones)
+TArray<FRigidAssetBoneGeometry> UBoneGeometryGenerator_Sphere::Build(FRigidAssetBoneSelection Bones)
 {
+	using namespace UE::Chaos::RigidAsset;
+
 	if(!Bones.Mesh || !Bones.Skeleton)
 	{
 		return {};
 	}
 
-	TArray<UE::Chaos::RigidAsset::FSimpleGeometry> Result;
+	TArray<FRigidAssetBoneGeometry> Result;
 
 	for(const FRigidAssetBoneInfo& Bone : Bones.SelectedBones)
 	{
-		FBoneVertData BoneVerts = GetBoneVertexData(Bone.Name, Bones.Mesh, BaseSettings.SourceLod, BaseSettings.VertexMode);
+		FBoneVertData Verts = GetMergedBoneVertexData(Bone, Bones.Mesh, BaseSettings.SourceLod, BaseSettings.VertexMode);
 
-		if(BoneVerts.Positions.Num() == 0)
+		if(Verts.Positions.Num() == 0)
 		{
 			// No influences, make a small sphere and warn
 			// #TODO pass context through for dataflow warnings
-			Result.AddDefaulted();
-			Result.Last().GeometryElement = MakeShared<FKSphereElem>(5.0f);
+			Result.Emplace(Bone, FSimpleGeometry(MakeShared<FKSphereElem>(5.0f)));
 			continue;
 		}
 
-		FBox3f VertBox(BoneVerts.Positions);
+		FBox3f VertBox(Verts.Positions);
 
 		TSharedPtr<FKSphereElem> NewSphere = MakeShared<FKSphereElem>();
 		NewSphere->Radius = VertBox.GetExtent().Length();
 		NewSphere->Center = static_cast<FVector>(VertBox.GetCenter());
 
-		Result.AddDefaulted();
-		Result.Last().GeometryElement = NewSphere;
+		Result.Emplace(Bone, FSimpleGeometry(NewSphere));
 	}
 
 	return Result;
@@ -330,14 +497,15 @@ void UBoneGeometryGenerator_Sphere::DebugDraw(UE::Dataflow::FContext& Context, I
 	RefSkel.GetBoneAbsoluteTransforms(BoneTransforms);
 
 	int32 BoneColorSeed = 0;
+
 	for(const FRigidAssetBoneInfo& Bone : Bones.SelectedBones)
 	{
-		if(DrawVertsForBoneName != NAME_None && Bone.Name == DrawVertsForBoneName)
+		if(DrawVertsForBoneName != NAME_None && Bone.Name != DrawVertsForBoneName)
 		{
 			continue;
 		}
-
-		FBoneVertData BoneVerts = GetBoneVertexData(Bone.Name, Bones.Mesh, BaseSettings.SourceLod, BaseSettings.VertexMode);
+	
+		FBoneVertData BoneVerts = GetMergedBoneVertexData(Bone, Bones.Mesh, BaseSettings.SourceLod, BaseSettings.VertexMode);
 
 		DataflowRenderingInterface.SetColor(FLinearColor::MakeRandomSeededColor(BoneColorSeed++));
 		for(const FVector3f& Pos : BoneVerts.Positions)
@@ -423,25 +591,26 @@ namespace Private
 	}
 }
 
-TArray<UE::Chaos::RigidAsset::FSimpleGeometry> UBoneGeometryGenerator_Capsule::Build(FRigidAssetBoneSelection Bones)
+TArray<FRigidAssetBoneGeometry> UBoneGeometryGenerator_Capsule::Build(FRigidAssetBoneSelection Bones)
 {
+	using namespace UE::Chaos::RigidAsset;
+
 	if(!Bones.Mesh || !Bones.Skeleton)
 	{
 		return {};
 	}
 
-	TArray<UE::Chaos::RigidAsset::FSimpleGeometry> Result;
+	TArray<FRigidAssetBoneGeometry> Result;
 
 	for(const FRigidAssetBoneInfo& Bone : Bones.SelectedBones)
 	{
-		FBoneVertData BoneVerts = GetBoneVertexData(Bone.Name, Bones.Mesh, BaseSettings.SourceLod, BaseSettings.VertexMode);
+		FBoneVertData BoneVerts = GetMergedBoneVertexData(Bone, Bones.Mesh, BaseSettings.SourceLod, BaseSettings.VertexMode);
 
 		if(BoneVerts.Positions.Num() == 0)
 		{
 			// No influences, make a small sphere and warn
 			// #TODO pass context through for dataflow warnings
-			Result.AddDefaulted();
-			Result.Last().GeometryElement = MakeShared<FKSphereElem>(5.0f);
+			Result.Emplace(Bone, FSimpleGeometry(MakeShared<FKSphereElem>(5.0f)));
 			continue;
 		}
 
@@ -532,8 +701,7 @@ TArray<UE::Chaos::RigidAsset::FSimpleGeometry> UBoneGeometryGenerator_Capsule::B
 		NewCapsule->Radius = Radius * 0.5f;
 		NewCapsule->Rotation = ElemMatrix.Rotator();
 
-		Result.AddDefaulted();
-		Result.Last().GeometryElement = NewCapsule;
+		Result.Emplace(Bone, FSimpleGeometry(NewCapsule));
 	}
 
 	return Result;
@@ -582,14 +750,16 @@ void UBoneGeometryGenerator_Capsule::DebugDraw(UE::Dataflow::FContext& Context,
 
 ////////////////////////////////////////////////////////////////////////////////////////////////////////////////
 
-TArray<UE::Chaos::RigidAsset::FSimpleGeometry> UBoneGeometryGenerator_Convex::Build(FRigidAssetBoneSelection Bones)
+TArray<FRigidAssetBoneGeometry> UBoneGeometryGenerator_Convex::Build(FRigidAssetBoneSelection Bones)
 {
+	using namespace UE::Chaos::RigidAsset;
+
 	if(!Bones.Mesh || !Bones.Skeleton)
 	{
 		return {};
 	}
 
-	TArray<UE::Chaos::RigidAsset::FSimpleGeometry> Result;
+	TArray<FRigidAssetBoneGeometry> Result;
 
 	for(const FRigidAssetBoneInfo& Bone : Bones.SelectedBones)
 	{
@@ -599,8 +769,7 @@ TArray<UE::Chaos::RigidAsset::FSimpleGeometry> UBoneGeometryGenerator_Convex::Bu
 		{
 			// No influences, make a small sphere and warn
 			// #TODO pass context through for dataflow warnings
-			Result.AddDefaulted();
-			Result.Last().GeometryElement = MakeShared<FKSphereElem>(5.0f);
+			Result.Emplace(Bone, FSimpleGeometry(MakeShared<FKSphereElem>(5.0f)));
 			continue;
 		}
 
@@ -630,8 +799,7 @@ TArray<UE::Chaos::RigidAsset::FSimpleGeometry> UBoneGeometryGenerator_Convex::Bu
 
 		NewConvex->UpdateElemBox();
 
-		Result.AddDefaulted();
-		Result.Last().GeometryElement = NewConvex;
+		Result.Emplace(Bone, FSimpleGeometry(NewConvex));
 	}
 
 	return Result;
diff --git a/Engine/Plugins/Experimental/ChaosRigidAsset/Source/ChaosRigidAssetNodes/Private/Generators/BoneGeometryGenerators.h b/Engine/Plugins/Experimental/ChaosRigidAsset/Source/ChaosRigidAssetNodes/Private/Generators/BoneGeometryGenerators.h
--- a/Engine/Plugins/Experimental/ChaosRigidAsset/Source/ChaosRigidAssetNodes/Private/Generators/BoneGeometryGenerators.h
+++ b/Engine/Plugins/Experimental/ChaosRigidAsset/Source/ChaosRigidAssetNodes/Private/Generators/BoneGeometryGenerators.h
@@ -14,6 +14,18 @@ namespace UE::Dataflow
 
 ////////////////////////////////////////////////////////////////////////////////////////////////////////////////
 
+/** 
+ * Merge operation to perform
+ */
+UENUM()
+enum class EMergeOperation
+{
+	MergeAll,
+	MergeSmall
+};
+
+////////////////////////////////////////////////////////////////////////////////////////////////////////////////
+
 /**
  * Operation to use for small bones
  * Skip will ignore a bone, and all of its geometry when attempting to generate collision shapes
@@ -43,6 +55,18 @@ enum class EVertexSelectMode
 
 ////////////////////////////////////////////////////////////////////////////////////////////////////////////////
 
+/**
+ * Whether to include merged bones in returned vertext data
+ */
+ UENUM()
+enum class EVertexMergeMode
+{
+	RootOnly,
+	Merged
+};
+
+////////////////////////////////////////////////////////////////////////////////////////////////////////////////
+
 /**
  * How to align a body during generation
  * Bone will align along the bone direction and wrap the shape around the verts
@@ -65,16 +89,19 @@ struct FBaseGenerationSettings
 {
 	GENERATED_BODY()
 
-	// Bones under this limit will be classified as a "Small" bone and apply the requested operation
 	UPROPERTY(EditAnywhere, Category = "Generation")
+	EMergeOperation Operation = EMergeOperation::MergeSmall;
+
+	// Bones under this limit will be classified as a "Small" bone and apply the requested operation
+	UPROPERTY(EditAnywhere, Category = "Generation", meta=(EditCondition="Operation==EMergeOperation::MergeSmall", EditConditionHides, Units=cm))
 	float MinimumBoneSize = 20.0f;
 
 	// Which operation to apply to small bones
-	UPROPERTY(EditAnywhere, Category = "Generation")
+	UPROPERTY(EditAnywhere, Category = "Generation", meta = (EditCondition = "Operation==EMergeOperation::MergeSmall", EditConditionHides))
 	ESmallBoneOperation SmallBoneOp = ESmallBoneOperation::Skip;
 
 	// For the "Merge" small-bone operation - only merge over this threshold and skip otherwise
-	UPROPERTY(EditAnywhere, Category = "Generation")
+	UPROPERTY(EditAnywhere, Category = "Generation", meta = (EditCondition = "Operation==EMergeOperation::MergeSmall", EditConditionHides, Units=cm))
 	float SmallBoneMergeThresholdOverride = 0.0f;
 
 	// Vertex extraction operation, see EVertexSelectMode
@@ -86,10 +113,34 @@ struct FBaseGenerationSettings
 	int32 SourceLod = 0;
 
 	// Extra thickness added to generated geometry after the genration process produces a geometry
-	UPROPERTY(EditAnywhere, Category = "Generation")
+	UPROPERTY(EditAnywhere, Category = "Generation", meta=(Units=cm))
 	float Thickness = 0.0f;
 };
 
+/**
+ * Data extracted from a bone for body generation
+ */
+struct FBoneVertData
+{
+	TArray<FVector3f> Positions;
+	TArray<FVector3f> Normals;
+	TArray<float> Weights;
+	FBoxSphereBounds3f Bounds;
+};
+
+struct FRigidAssetBoneGeometry
+{
+	FRigidAssetBoneInfo Bone;
+	UE::Chaos::RigidAsset::FSimpleGeometry Geometry;
+};
+
+/**
+ * Given a bone selection, produce a set of merged bones according the the specified generation settings
+ * @param Settings Generation setting to apply for the merge algorithm
+ * @param InSelectionSorted A pre-sorted selection of bones
+ */
+FRigidAssetBoneSelection MergeSelection(const FBaseGenerationSettings& Settings, const FRigidAssetBoneSelection& InSelectionSorted);
+
 /**
  * Base geometry generator for each render->physics generation type.
  */
@@ -100,7 +151,7 @@ class UBoneGeometryGenerator : public UObject
 
 public:
 
-	virtual TArray<UE::Chaos::RigidAsset::FSimpleGeometry> Build(FRigidAssetBoneSelection Bones)
+	virtual TArray<FRigidAssetBoneGeometry> Build(FRigidAssetBoneSelection Bones)
 	{
 		return {};
 	}
@@ -131,7 +182,7 @@ class UBoneGeometryGenerator_Box : public UBoneGeometryGenerator
 
 public:
 
-	TArray<UE::Chaos::RigidAsset::FSimpleGeometry> Build(FRigidAssetBoneSelection Bones) final override;
+	TArray<FRigidAssetBoneGeometry> Build(FRigidAssetBoneSelection Bones) final override;
 
 private:
 };
@@ -149,7 +200,7 @@ class UBoneGeometryGenerator_Sphere : public UBoneGeometryGenerator
 
 public:
 
-	TArray<UE::Chaos::RigidAsset::FSimpleGeometry> Build(FRigidAssetBoneSelection Bones) final override;
+	TArray<FRigidAssetBoneGeometry> Build(FRigidAssetBoneSelection Bones) final override;
 
 #if WITH_EDITOR
 	virtual bool CanDebugDraw() const override;
@@ -184,7 +235,7 @@ class UBoneGeometryGenerator_Capsule : public UBoneGeometryGenerator
 
 public:
 
-	TArray<UE::Chaos::RigidAsset::FSimpleGeometry> Build(FRigidAssetBoneSelection Bones) final override;
+	TArray<FRigidAssetBoneGeometry> Build(FRigidAssetBoneSelection Bones) final override;
 
 #if WITH_EDITOR
 	virtual bool CanDebugDraw() const override;
@@ -219,7 +270,7 @@ class UBoneGeometryGenerator_Convex : public UBoneGeometryGenerator
 
 public:
 
-	TArray<UE::Chaos::RigidAsset::FSimpleGeometry> Build(FRigidAssetBoneSelection Bones) final override;
+	TArray<FRigidAssetBoneGeometry> Build(FRigidAssetBoneSelection Bones) final override;
 
 private:
 };
diff --git a/Engine/Plugins/Experimental/ChaosRigidAsset/Source/ChaosRigidAssetNodes/Private/PhysicsAssetBuilder.cpp b/Engine/Plugins/Experimental/ChaosRigidAsset/Source/ChaosRigidAssetNodes/Private/PhysicsAssetBuilder.cpp
--- a/Engine/Plugins/Experimental/ChaosRigidAsset/Source/ChaosRigidAssetNodes/Private/PhysicsAssetBuilder.cpp
+++ b/Engine/Plugins/Experimental/ChaosRigidAsset/Source/ChaosRigidAssetNodes/Private/PhysicsAssetBuilder.cpp
@@ -5,7 +5,8 @@
 #include "PhysicsEngine/SkeletalBodySetup.h"
 #include "PhysicsEngine/PhysicsConstraintTemplate.h"
 #include "Engine/SkeletalMesh.h"
-#include "UObject/Package.h"
+#include "UObject/UObjectIterator.h"
+#include "Components/SkeletalMeshComponent.h"
 
 #if WITH_EDITOR
 #include "IAssetTools.h"
@@ -146,7 +147,26 @@ namespace UE::Chaos::RigidAsset
 			TargetAsset->DisableCollision(Indices[0], Indices[1]);
 		}
 
-		RefreshSkelMeshOnPhysicsAssetChange(TargetAsset->GetPreviewMesh());
+		// Refresh components using this physics asset
+		for(FThreadSafeObjectIterator Iter(USkeletalMeshComponent::StaticClass()); Iter; ++Iter)
+		{
+			USkeletalMeshComponent* SkeletalMeshComponent = Cast<USkeletalMeshComponent>(*Iter);
+
+			if(!SkeletalMeshComponent || !SkeletalMeshComponent->IsRegistered())
+			{
+				// If the component is invalid, or it's not registered then it can't have physics state.
+				continue;
+			}
+
+			if(SkeletalMeshComponent->GetPhysicsAsset() == TargetAsset)
+			{
+				if(SkeletalMeshComponent->IsPhysicsStateCreated())
+				{
+					SkeletalMeshComponent->RecreatePhysicsState();
+					SkeletalMeshComponent->UpdateHasValidBodies();
+				}
+			}
+		}
 
 		return TargetAsset;
 	}
diff --git a/Engine/Plugins/Experimental/ChaosRigidAsset/Source/ChaosRigidAssetNodes/Private/PhysicsAssetDataflowNodes.cpp b/Engine/Plugins/Experimental/ChaosRigidAsset/Source/ChaosRigidAssetNodes/Private/PhysicsAssetDataflowNodes.cpp
--- a/Engine/Plugins/Experimental/ChaosRigidAsset/Source/ChaosRigidAssetNodes/Private/PhysicsAssetDataflowNodes.cpp
+++ b/Engine/Plugins/Experimental/ChaosRigidAsset/Source/ChaosRigidAssetNodes/Private/PhysicsAssetDataflowNodes.cpp
@@ -23,6 +23,8 @@
 #include "AnimationUtils.h"
 #include "ChaosVDRuntimeModule.h"
 
+#define LOCTEXT_NAMESPACE "PhysicsAssetDataflowNodes"
+
 FDataflowPhysicsAssetTerminalNode::FDataflowPhysicsAssetTerminalNode(const UE::Dataflow::FNodeParameters& InParam, FGuid InGuid /*= FGuid::NewGuid()*/)
 	: FDataflowTerminalNode(InParam, InGuid)
 {
@@ -82,6 +84,7 @@ void UE::Dataflow::RegisterPhysicsAssetNodes()
 	DATAFLOW_NODE_REGISTER_CREATION_FACTORY(FDataflowAppendBoneSelection);
 	DATAFLOW_NODE_REGISTER_CREATION_FACTORY(FDataflowSelectBonesByName);
 	DATAFLOW_NODE_REGISTER_CREATION_FACTORY(FDataflowSelectConnectedBones);
+	DATAFLOW_NODE_REGISTER_CREATION_FACTORY(FDataflowMergeBoneSelection);
 
 	DATAFLOW_NODE_REGISTER_CREATION_FACTORY(FDataflowRigidSimulationCachingNode);
 	
@@ -292,6 +295,7 @@ private:
 };
 
 #if WITH_EDITOR
+
 bool FDataflowMakeBody::CanDebugDraw() const
 {
 	return true;
@@ -327,7 +331,7 @@ void FDataflowMakeBody::Evaluate(UE::Dataflow::FContext& Context, const FDataflo
 void FDataflowMakeBody::Register(const UE::Dataflow::FNodeParameters& InParam)
 {
 	// Construct the template object for this node
-	Body = NewObject<USkeletalBodySetup>(GetOwner());
+	Body = NewObject<USkeletalBodySetup>(GetOwner(), MakeUniqueObjectName(GetOwner(), USkeletalBodySetup::StaticClass(), "TemplateBody"));
 
 	AddInputs(BoneName);
 	AddOutputs(Body);
@@ -513,6 +517,30 @@ void FDataflowSelectConnectedBones::Register(const UE::Dataflow::FNodeParameters
 	AddPassthroughs(Selection);
 }
 
+void FDataflowMergeBoneSelection::Evaluate(UE::Dataflow::FContext& Context, const FDataflowOutput* Out) const
+{
+	if(Out->IsA(&Selection))
+	{
+		const FRigidAssetBoneSelection& InSelection = GetValue(Context, &Selection);
+
+		if(!InSelection.Skeleton)
+		{
+			Context.Error(LOCTEXT("InvalidSelection_NoSkeleton", "Attempted to make a selection without a Skeleton"), this);
+			SafeForwardInput(Context, &Selection, &Selection);
+			return;
+		}
+
+		FRigidAssetBoneSelection NewSelection = MergeSelection(MergeProfile, InSelection);
+		SetValue(Context, NewSelection, &Selection);
+	}
+}
+
+void FDataflowMergeBoneSelection::Register(const UE::Dataflow::FNodeParameters& InParam)
+{
+	AddInputs(Selection);
+	AddPassthroughs(Selection);
+}
+
 #if WITH_EDITOR
 bool FDataflowCreateGeometryForBones::CanDebugDraw() const
 {
@@ -529,6 +557,13 @@ void FDataflowCreateGeometryForBones::DebugDraw(UE::Dataflow::FContext& Context,
 	const FRigidAssetBoneSelection& InSelection = GetValue(Context, &Selection);
 	TRefCountPtr<IDataflowDebugDrawObject> DrawObject(MakeDebugDrawObject<FDataflowDebugDrawBoneSelectionObject>(DataflowRenderingInterface.ModifyDataflowElements(), InSelection));
 	DataflowRenderingInterface.DrawObject(DrawObject);
+
+	TObjectPtr<UBoneGeometryGenerator> InGenerator = GetValue(Context, &Generator);
+
+	if(InGenerator)
+	{
+		InGenerator->DebugDraw(Context, DataflowRenderingInterface, DebugDrawParameters, InSelection);
+	}
 }
 #endif
 
@@ -563,7 +598,7 @@ void FDataflowCreateGeometryForBones::Evaluate(UE::Dataflow::FContext& Context,
 		}
 
 		FPhysicsAssetDataflowState InState = GetValue(Context, &State);
-		TArray<UE::Chaos::RigidAsset::FSimpleGeometry> Geoms = InGenerator->Build(InSelection);
+		TArray<FRigidAssetBoneGeometry> Geoms = InGenerator->Build(InSelection);
 
 		if(Geoms.Num() != InSelection.SelectedBones.Num())
 		{
@@ -574,8 +609,8 @@ void FDataflowCreateGeometryForBones::Evaluate(UE::Dataflow::FContext& Context,
 
 		for(int32 Index = 0; Index < Geoms.Num(); ++Index)
 		{
-			const UE::Chaos::RigidAsset::FSimpleGeometry& Geom = Geoms[Index];
-			const FRigidAssetBoneInfo& Bone = InSelection.SelectedBones[Index];
+			const UE::Chaos::RigidAsset::FSimpleGeometry& Geom = Geoms[Index].Geometry;
+			const FRigidAssetBoneInfo& Bone = Geoms[Index].Bone;
 
 			TObjectPtr<USkeletalBodySetup> NewSetup = CastChecked<USkeletalBodySetup>(StaticDuplicateObject(InTemplate, GetTransientPackage()));
 
@@ -1100,3 +1135,5 @@ void FDataflowAggGeomAddShape::Register()
 	AddInputs(AggGeom);
 	AddPassthroughs(AggGeom);
 }
+
+#undef LOCTEXT_NAMESPACE
diff --git a/Engine/Plugins/Experimental/ChaosRigidAsset/Source/ChaosRigidAssetNodes/Private/PhysicsAssetDataflowNodes.h b/Engine/Plugins/Experimental/ChaosRigidAsset/Source/ChaosRigidAssetNodes/Private/PhysicsAssetDataflowNodes.h
--- a/Engine/Plugins/Experimental/ChaosRigidAsset/Source/ChaosRigidAssetNodes/Private/PhysicsAssetDataflowNodes.h
+++ b/Engine/Plugins/Experimental/ChaosRigidAsset/Source/ChaosRigidAssetNodes/Private/PhysicsAssetDataflowNodes.h
@@ -359,6 +359,40 @@ private:
 
 ////////////////////////////////////////////////////////////////////////////////////////////////////////////////
 
+/**
+ * Merge a bone selection using the specified merge profile
+ */
+USTRUCT()
+struct FDataflowMergeBoneSelection : public FRigidDataflowNode
+{
+	GENERATED_BODY()
+	DATAFLOW_NODE_DEFINE_INTERNAL(FDataflowMergeBoneSelection, "Merge Bones in Selection", "PhysicsAsset", "")
+	DATAFLOW_NODE_RENDER_TYPE("GeomRender", FName("FBoneSelection"), "Selection")
+
+public:
+
+	FDataflowMergeBoneSelection(const UE::Dataflow::FNodeParameters& InParam, FGuid InGuid = FGuid::NewGuid())
+		: FRigidDataflowNode(InParam, InGuid)
+	{
+		Register(InParam);
+	}
+
+private:
+
+	// Resulting bone selection
+	UPROPERTY(meta = (DataflowInput, DataflowOutput, DataflowPassthrough = "Selection"))
+	FRigidAssetBoneSelection Selection;
+
+	// Merge settings to use to decide how to merge the bones
+	UPROPERTY(EditAnywhere, Category = Settings)
+	FBaseGenerationSettings MergeProfile;
+
+	void Evaluate(UE::Dataflow::FContext& Context, const FDataflowOutput* Out) const override;
+	void Register(const UE::Dataflow::FNodeParameters& InParam);
+};
+
+////////////////////////////////////////////////////////////////////////////////////////////////////////////////
+
 /**
  * Apply a geometry generator to a bone selection, creating bodies for each bone
  * The method of geometry generation is generic and offloaded to a "Generator" object/
diff --git a/Engine/Plugins/Experimental/ChaosRigidAsset/Source/ChaosRigidAssetNodes/Private/ShapeElemNodes.cpp b/Engine/Plugins/Experimental/ChaosRigidAsset/Source/ChaosRigidAssetNodes/Private/ShapeElemNodes.cpp
--- a/Engine/Plugins/Experimental/ChaosRigidAsset/Source/ChaosRigidAssetNodes/Private/ShapeElemNodes.cpp
+++ b/Engine/Plugins/Experimental/ChaosRigidAsset/Source/ChaosRigidAssetNodes/Private/ShapeElemNodes.cpp
@@ -66,7 +66,7 @@ namespace UE::Chaos::RigidAsset
 				NewElem = MakeShared<ElemType>(*Elem);
 			});
 
-		return { LocalTransform, MoveTemp(NewElem) };
+		return { LocalTransform, NewElem };
 	}
 
 	void FMakeSphereElemDataflowNode::Evaluate(UE::Dataflow::FContext& Context, const FDataflowOutput* Out) const
diff --git a/Engine/Plugins/Experimental/ChaosRigidAsset/Source/ChaosRigidAssetNodes/Private/ShapeElemNodes.h b/Engine/Plugins/Experimental/ChaosRigidAsset/Source/ChaosRigidAssetNodes/Private/ShapeElemNodes.h
--- a/Engine/Plugins/Experimental/ChaosRigidAsset/Source/ChaosRigidAssetNodes/Private/ShapeElemNodes.h
+++ b/Engine/Plugins/Experimental/ChaosRigidAsset/Source/ChaosRigidAssetNodes/Private/ShapeElemNodes.h
@@ -26,6 +26,19 @@ namespace UE::Chaos::RigidAsset
 	{
 		GENERATED_BODY();
 
+		FSimpleGeometry() = default;
+
+		FSimpleGeometry(TSharedPtr<FKShapeElem> InElem)
+			: GeometryElement(InElem)
+		{
+		}
+
+		FSimpleGeometry(const FTransform InLocalTransform, TSharedPtr<FKShapeElem> InElem)
+			: LocalTransform(InLocalTransform)
+			, GeometryElement(InElem)
+		{
+		}
+
 		// Duplicate this geometry handling the geometry type correctly
 		FSimpleGeometry Duplicate();
 
diff --git a/Engine/Plugins/Experimental/ChaosRigidAsset/Source/ChaosRigidAssetNodes/Public/BoneSelection.h b/Engine/Plugins/Experimental/ChaosRigidAsset/Source/ChaosRigidAssetNodes/Public/BoneSelection.h
--- a/Engine/Plugins/Experimental/ChaosRigidAsset/Source/ChaosRigidAssetNodes/Public/BoneSelection.h
+++ b/Engine/Plugins/Experimental/ChaosRigidAsset/Source/ChaosRigidAssetNodes/Public/BoneSelection.h
@@ -33,6 +33,10 @@ struct FRigidAssetBoneInfo
 	UPROPERTY()
 	int32 Depth = INDEX_NONE;
 
+	// List of child bones merged into this one
+	UPROPERTY()
+	TArray<int32> MergedChildren;
+
 	bool operator ==(const FRigidAssetBoneInfo& Other) const = default;
 	bool operator <(const FRigidAssetBoneInfo& Other) const;
 };
