Non-unity fix

#Synced-CL 48950948

diff --git a/Engine/Plugins/Experimental/ChaosRigidAsset/Source/ChaosRigidAssetNodes/Private/PhysicsAssetBuilder.cpp b/Engine/Plugins/Experimental/ChaosRigidAsset/Source/ChaosRigidAssetNodes/Private/PhysicsAssetBuilder.cpp
--- a/Engine/Plugins/Experimental/ChaosRigidAsset/Source/ChaosRigidAssetNodes/Private/PhysicsAssetBuilder.cpp
+++ b/Engine/Plugins/Experimental/ChaosRigidAsset/Source/ChaosRigidAssetNodes/Private/PhysicsAssetBuilder.cpp
@@ -7,6 +7,7 @@
 #include "Engine/SkeletalMesh.h"
 #include "UObject/UObjectIterator.h"
 #include "Components/SkeletalMeshComponent.h"
+#include "UObject/Package.h"
 
 #if WITH_EDITOR
 #include "IAssetTools.h"
