[Chaos] [Dataflow] Allow convex geometry generator to use merged bone selections

#Synced-CL 48990935

diff --git a/Engine/Plugins/Experimental/ChaosRigidAsset/Source/ChaosRigidAssetNodes/Private/Generators/BoneGeometryGenerators.cpp b/Engine/Plugins/Experimental/ChaosRigidAsset/Source/ChaosRigidAssetNodes/Private/Generators/BoneGeometryGenerators.cpp
--- a/Engine/Plugins/Experimental/ChaosRigidAsset/Source/ChaosRigidAssetNodes/Private/Generators/BoneGeometryGenerators.cpp
+++ b/Engine/Plugins/Experimental/ChaosRigidAsset/Source/ChaosRigidAssetNodes/Private/Generators/BoneGeometryGenerators.cpp
@@ -763,7 +763,7 @@ TArray<FRigidAssetBoneGeometry> UBoneGeometryGenerator_Convex::Build(FRigidAsset
 
 	for(const FRigidAssetBoneInfo& Bone : Bones.SelectedBones)
 	{
-		FBoneVertData BoneVerts = GetBoneVertexData(Bone.Name, Bones.Mesh, BaseSettings.SourceLod, BaseSettings.VertexMode);
+		FBoneVertData BoneVerts = GetMergedBoneVertexData(Bone, Bones.Mesh, BaseSettings.SourceLod, BaseSettings.VertexMode);
 
 		if(BoneVerts.Positions.Num() == 0)
 		{
