[Chaos] [Dataflow] Added motion types to s/t constraint node so the node will produce a constraint with locked axes when needed

#Synced-CL 48990935

diff --git a/Engine/Plugins/Experimental/ChaosRigidAsset/Source/ChaosRigidAssetNodes/Private/Generators/ConstraintGenerators.cpp b/Engine/Plugins/Experimental/ChaosRigidAsset/Source/ChaosRigidAssetNodes/Private/Generators/ConstraintGenerators.cpp
--- a/Engine/Plugins/Experimental/ChaosRigidAsset/Source/ChaosRigidAssetNodes/Private/Generators/ConstraintGenerators.cpp
+++ b/Engine/Plugins/Experimental/ChaosRigidAsset/Source/ChaosRigidAssetNodes/Private/Generators/ConstraintGenerators.cpp
@@ -18,6 +18,21 @@ namespace UE::Dataflow
 
 ////////////////////////////////////////////////////////////////////////////////////////////////////////////////
 
+EAngularConstraintMotion GetMotionType(float InLimit)
+{
+	if(InLimit == 0)
+	{
+		return ACM_Locked;
+	}
+
+	if(InLimit >= 180.0f)
+	{
+		return ACM_Free;
+	}
+
+	return ACM_Limited;
+}
+
 TArray<TObjectPtr<UPhysicsConstraintTemplate>> UConstraintGenerator_SwingTwist::Build(TObjectPtr<UPhysicsConstraintTemplate> ConstraintTemplate, FRigidAssetBoneSelection Bones) const
 {
 	TArray<TObjectPtr<UPhysicsConstraintTemplate>> Result;
@@ -50,9 +65,9 @@ TArray<TObjectPtr<UPhysicsConstraintTemplate>> UConstraintGenerator_SwingTwist::
 				Defaults.CopyConstraintParamsFrom(&ConstraintTemplate->DefaultInstance);
 			}
 
-			Defaults.SetAngularSwing1Limit(EAngularConstraintMotion::ACM_Limited, Swing1Limit);
-			Defaults.SetAngularSwing2Limit(EAngularConstraintMotion::ACM_Limited, Swing2Limit);
-			Defaults.SetAngularTwistLimit(EAngularConstraintMotion::ACM_Limited, TwistLimit);
+			Defaults.SetAngularSwing1Limit(GetMotionType(Swing1Limit), Swing1Limit);
+			Defaults.SetAngularSwing2Limit(GetMotionType(Swing2Limit), Swing2Limit);
+			Defaults.SetAngularTwistLimit(GetMotionType(TwistLimit), TwistLimit);
 
 			Defaults.JointName = Bone.Name;
 			Defaults.ConstraintBone1 = Bone.Name;
