[Dataflow] [RigidAsset] Added convex decomposition node to dataflow for physics assets. Supports a simple number-of-hulls approach and a more configurable negative space mode using GeometryTools decomposition systems.

Extended API on CopyOnWriteArray for needs from new nodes

#Synced-CL 49921579

diff --git a/Engine/Plugins/Experimental/ChaosRigidAsset/ChaosRigidAsset.uplugin b/Engine/Plugins/Experimental/ChaosRigidAsset/ChaosRigidAsset.uplugin
--- a/Engine/Plugins/Experimental/ChaosRigidAsset/ChaosRigidAsset.uplugin
+++ b/Engine/Plugins/Experimental/ChaosRigidAsset/ChaosRigidAsset.uplugin
@@ -37,6 +37,10 @@
 		{
 			"Name": "Dataflow",
 			"Enabled": true
+		},
+		{
+			"Name": "GeometryProcessing",
+			"Enabled": true
 		}
 	]
 }
diff --git a/Engine/Plugins/Experimental/ChaosRigidAsset/Source/ChaosRigidAssetEditor/Private/DataflowRendering.cpp b/Engine/Plugins/Experimental/ChaosRigidAsset/Source/ChaosRigidAssetEditor/Private/DataflowRendering.cpp
--- a/Engine/Plugins/Experimental/ChaosRigidAsset/Source/ChaosRigidAssetEditor/Private/DataflowRendering.cpp
+++ b/Engine/Plugins/Experimental/ChaosRigidAsset/Source/ChaosRigidAssetEditor/Private/DataflowRendering.cpp
@@ -344,9 +344,9 @@ namespace UE::Chaos::RigidAsset
 
 			if(AssetState.HasData())
 			{
-				const FReferenceSkeleton& RefSkeleton = AssetState.TargetMesh->GetRefSkeleton();
+				const FReferenceSkeleton& RefSkeleton = AssetState.GetMesh()->GetRefSkeleton();
 
-				for(TObjectPtr<USkeletalBodySetup> Body : AssetState.Bodies)
+				for(TObjectPtr<const USkeletalBodySetup> Body : AssetState.GetBodies())
 				{
 					const int32 BoneIndex = RefSkeleton.FindBoneIndex(Body->BoneName);
 
diff --git a/Engine/Plugins/Experimental/ChaosRigidAsset/Source/ChaosRigidAssetNodes/ChaosRigidAssetNodes.Build.cs b/Engine/Plugins/Experimental/ChaosRigidAsset/Source/ChaosRigidAssetNodes/ChaosRigidAssetNodes.Build.cs
--- a/Engine/Plugins/Experimental/ChaosRigidAsset/Source/ChaosRigidAssetNodes/ChaosRigidAssetNodes.Build.cs
+++ b/Engine/Plugins/Experimental/ChaosRigidAsset/Source/ChaosRigidAssetNodes/ChaosRigidAssetNodes.Build.cs
@@ -15,6 +15,8 @@ public class ChaosRigidAssetNodes : ModuleRules
 				"Engine",
 				"DataflowCore",
 				"DataflowEngine",
+				"GeometryCore",
+				"DynamicMesh",
 				"ChaosRigidAssetEngine"
 			}
 		);
diff --git a/Engine/Plugins/Experimental/ChaosRigidAsset/Source/ChaosRigidAssetNodes/Internal/PhysicsAssetDataflowState.h b/Engine/Plugins/Experimental/ChaosRigidAsset/Source/ChaosRigidAssetNodes/Internal/PhysicsAssetDataflowState.h
--- a/Engine/Plugins/Experimental/ChaosRigidAsset/Source/ChaosRigidAssetNodes/Internal/PhysicsAssetDataflowState.h
+++ b/Engine/Plugins/Experimental/ChaosRigidAsset/Source/ChaosRigidAssetNodes/Internal/PhysicsAssetDataflowState.h
@@ -18,7 +18,7 @@ class FReferenceCollector;
  * Represents the complete state required to build a physics asset. Datflow nodes populate this container
  * until it reaches a terminal which packages it into a final UPhysicsAsset. This allows us to keep the
  * const invariant for dataflow evaluation by adding copy-on-write support to the state.
- * During node evaluation this object should be passed by-value as it is a collection of pts to the 
+ * During node evaluation this object should be passed by-value as it is a collection of pts to the
  * actual data. If the copy-on-write arrays are modified they will automatically duplicate if required
  */
 USTRUCT()
@@ -26,14 +26,22 @@ struct FPhysicsAssetDataflowState
 {
 	GENERATED_BODY()
 
-	void AddStructReferencedObjects(FReferenceCollector& Collector);
+public:
 
-	// Base asset data
-	TObjectPtr<USkeleton> TargetSkeleton = nullptr;
-	TObjectPtr<USkeletalMesh> TargetMesh = nullptr;
+	CHAOSRIGIDASSETNODES_API FPhysicsAssetDataflowState();
+	explicit CHAOSRIGIDASSETNODES_API FPhysicsAssetDataflowState(TObjectPtr<USkeleton> TargetSkeleton, TObjectPtr<USkeletalMesh> TargetMesh);
+
+	CHAOSRIGIDASSETNODES_API USkeleton* GetSkeleton() const;
+	CHAOSRIGIDASSETNODES_API USkeletalMesh* GetMesh() const;
+	CHAOSRIGIDASSETNODES_API USkeletalBodySetup* GetBody(FName BoneName) const;
+	CHAOSRIGIDASSETNODES_API void SetBody(FName BoneName, USkeletalBodySetup* InSetup);
+	CHAOSRIGIDASSETNODES_API USkeletalBodySetup* FindOrCreateBody(FName BoneName, USkeletalBodySetup* TemplateBody = nullptr);
+	CHAOSRIGIDASSETNODES_API USkeletalBodySetup* FindOrCreateBodyChecked(FName BoneName, USkeletalBodySetup* TemplateBody = nullptr);
+	CHAOSRIGIDASSETNODES_API const TCopyOnWriteArray<TObjectPtr<USkeletalBodySetup>>& GetBodies() const;
+
+	void AddStructReferencedObjects(FReferenceCollector& Collector);
 
 	// Bodies and constraints in the simulation setup
-	TCopyOnWriteArray<TObjectPtr<USkeletalBodySetup>> Bodies;
 	TCopyOnWriteArray<TObjectPtr<UPhysicsConstraintTemplate>> Constraints;
 
 	void DebugLog() const;
@@ -41,6 +49,12 @@ struct FPhysicsAssetDataflowState
 
 private:
 
+	// Base asset data
+	TObjectPtr<USkeleton> TargetSkeleton = nullptr;
+	TObjectPtr<USkeletalMesh> TargetMesh = nullptr;
+
+	TCopyOnWriteArray<int32> BoneIndexToBody;
+	TCopyOnWriteArray<TObjectPtr<USkeletalBodySetup>> Bodies;
 };
 
 /**
diff --git a/Engine/Plugins/Experimental/ChaosRigidAsset/Source/ChaosRigidAssetNodes/Private/Generators/BoneGeometryGenerators.cpp b/Engine/Plugins/Experimental/ChaosRigidAsset/Source/ChaosRigidAssetNodes/Private/Generators/BoneGeometryGenerators.cpp
--- a/Engine/Plugins/Experimental/ChaosRigidAsset/Source/ChaosRigidAssetNodes/Private/Generators/BoneGeometryGenerators.cpp
+++ b/Engine/Plugins/Experimental/ChaosRigidAsset/Source/ChaosRigidAssetNodes/Private/Generators/BoneGeometryGenerators.cpp
@@ -16,6 +16,9 @@
 #include "Rendering/SkeletalMeshRenderData.h"
 #include "Rendering/SkinWeightVertexBuffer.h"
 #include "Rendering/StaticMeshVertexBuffer.h"
+#include "Operations/MeshConvexHull.h"
+
+DEFINE_LOG_CATEGORY(LogBoneGeometryGenerators);
 
 namespace UE::Dataflow
 {
@@ -26,6 +29,7 @@ namespace UE::Dataflow
 		DATAFLOW_NODE_REGISTER_CREATION_FACTORY(FMakeSphereBoneGeometryGenerator);
 		DATAFLOW_NODE_REGISTER_CREATION_FACTORY(FMakeCapsuleBoneGeometryGenerator);
 		DATAFLOW_NODE_REGISTER_CREATION_FACTORY(FMakeConvexBoneGeometryGenerator);
+		DATAFLOW_NODE_REGISTER_CREATION_FACTORY(FMakeConvexDecompBoneGeometryGenerator);
 	}
 }
 
@@ -103,6 +107,24 @@ void FMakeConvexBoneGeometryGenerator::Register()
 
 ////////////////////////////////////////////////////////////////////////////////////////////////////////////////
 
+void FMakeConvexDecompBoneGeometryGenerator::Evaluate(UE::Dataflow::FContext& Context, const FDataflowOutput* Out) const
+{
+	if(Out->IsA(&Generator))
+	{
+		TObjectPtr<UBoneGeometryGenerator> OutGenerator = CastChecked<UBoneGeometryGenerator>(StaticDuplicateObject(Generator, GetTransientPackage()));
+		SetValue(Context, OutGenerator, &Generator);
+	}
+}
+
+void FMakeConvexDecompBoneGeometryGenerator::Register()
+{
+	AddOutputs(Generator);
+
+	Generator = NewObject<UBoneGeometryGenerator_ConvexDecomposition>(GetOwner());
+}
+
+////////////////////////////////////////////////////////////////////////////////////////////////////////////////
+
 #if WITH_EDITOR
 bool UBoneGeometryGenerator::CanDebugDraw() const
 {
@@ -149,6 +171,12 @@ FBoneVertData GetBoneVertexData(FName BoneName, USkeletalMesh* Mesh, int32 Lod,
 	const FStaticMeshVertexBuffer& MeshVertexBuffer = LodData.StaticVertexBuffers.StaticMeshVertexBuffer;
 	const FPositionVertexBuffer& Positions = LodData.StaticVertexBuffers.PositionVertexBuffer;
 
+	TArray<uint32> Indices;
+	LodData.MultiSizeIndexContainer.GetIndexBuffer(Indices);
+
+	check(Indices.Num() % 3 == 0);
+	const int32 NumTris = Indices.Num() / 3;
+
 	const uint32 NumLodVerts = Weights->GetNumVertices();
 	const uint32 MaxInfluences = LodData.GetVertexBufferMaxBoneInfluences();
 
@@ -165,57 +193,124 @@ FBoneVertData GetBoneVertexData(FName BoneName, USkeletalMesh* Mesh, int32 Lod,
 		const FSkelMeshRenderSection& Section = LodData.RenderSections[SectionIndex];
 		const uint32 SectionMaxInfluences = Section.MaxBoneInfluences;
 
-		for(uint32 VertIndex = Section.BaseVertexIndex; VertIndex < Section.BaseVertexIndex + Section.NumVertices; ++VertIndex)
+		for(uint32 TriIndex = 0; TriIndex < Section.NumTriangles; ++TriIndex)
 		{
-			if(SelectMode == EVertexSelectMode::Any)
-			{
-				for(uint32 InfluenceIndex = 0; InfluenceIndex < SectionMaxInfluences; InfluenceIndex++)
-				{
-					const uint32 RenderBoneIndex = Weights->GetBoneIndex(VertIndex, InfluenceIndex);
-					const uint32 MeshBoneIndex = Section.BoneMap[RenderBoneIndex];
+			const uint32* TriBase = &Indices[Section.BaseIndex + TriIndex * 3];
+
+			bool bAcceptTriangle = false;
+			uint32 ReferenceBoneIndex = 0;
+			uint32 ReferenceInfluenceIndex = 0;
 
-					uint16 BoneWeight = Weights->GetBoneWeight(VertIndex, InfluenceIndex);
+			for(int32 TriVertIndex = 0; TriVertIndex < 3; ++TriVertIndex)
+			{
+				const uint32 VertIndex = TriBase[TriVertIndex];
 
-					if(BoneWeight > 0 && RefSkeleton.GetBoneName(MeshBoneIndex) == BoneName)
+				if(SelectMode == EVertexSelectMode::Any)
+				{
+					for(uint32 InfluenceIndex = 0; InfluenceIndex < SectionMaxInfluences; InfluenceIndex++)
 					{
-						const FMatrix44f& RefMatrix = Mesh->GetRefBasesInvMatrix()[MeshBoneIndex];
+						const uint32 RenderBoneIndex = Weights->GetBoneIndex(VertIndex, InfluenceIndex);
+						const uint32 MeshBoneIndex = Section.BoneMap[RenderBoneIndex];
 
-						Result.Positions.Add(RefMatrix.TransformPosition(Positions.VertexPosition(VertIndex)));
-						Result.Normals.Add(RefMatrix.TransformVector(MeshVertexBuffer.VertexTangentZ(VertIndex)));
-						Result.Weights.Add(BoneWeight * WeightsToFloatMult);
-						BoundsBuilder += Result.Positions.Last();
+						uint16 BoneWeight = Weights->GetBoneWeight(VertIndex, InfluenceIndex);
 
-						// Vertex will only have one influence for a given bone, skip the rest
-						break;
+						if(BoneWeight > 0 && RefSkeleton.GetBoneName(MeshBoneIndex) == BoneName)
+						{
+							bAcceptTriangle = true;
+							ReferenceBoneIndex = MeshBoneIndex;
+							ReferenceInfluenceIndex = InfluenceIndex;
+
+							// Vertex will only have one influence for a given bone, skip the rest
+							break;
+						}
 					}
 				}
-			}
-			else
-			{
-				int32 MaxInfluenceIndex = 0;
-				uint16 MaxInfluence = Weights->GetBoneWeight(VertIndex, 0);
-				for(uint32 InfluenceIndex = 1; InfluenceIndex < SectionMaxInfluences; InfluenceIndex++)
+				else
 				{
-					uint16 BoneWeight = Weights->GetBoneWeight(VertIndex, InfluenceIndex);
+					int32 MaxInfluenceIndex = 0;
+					uint16 MaxInfluence = Weights->GetBoneWeight(VertIndex, 0);
+					for(uint32 InfluenceIndex = 1; InfluenceIndex < SectionMaxInfluences; InfluenceIndex++)
+					{
+						uint16 BoneWeight = Weights->GetBoneWeight(VertIndex, InfluenceIndex);
+
+						if(BoneWeight > MaxInfluence)
+						{
+							MaxInfluenceIndex = InfluenceIndex;
+							MaxInfluence = BoneWeight;
+						}
+					}
 
-					if(BoneWeight > MaxInfluence)
+					const uint32 RenderBoneIndex = Weights->GetBoneIndex(VertIndex, MaxInfluenceIndex);
+					const uint32 MeshBoneIndex = Section.BoneMap[RenderBoneIndex];
+
+					if(RefSkeleton.GetBoneName(MeshBoneIndex) == BoneName)
 					{
-						MaxInfluenceIndex = InfluenceIndex;
-						MaxInfluence = BoneWeight;
+						bAcceptTriangle = true;
+						ReferenceBoneIndex = MeshBoneIndex;
+						ReferenceInfluenceIndex = MaxInfluenceIndex;
 					}
 				}
 
-				const uint32 RenderBoneIndex = Weights->GetBoneIndex(VertIndex, MaxInfluenceIndex);
-				const uint32 MeshBoneIndex = Section.BoneMap[RenderBoneIndex];
+				if(bAcceptTriangle)
+				{
+					// Once we have one vert we no longer need to check, this triangle is included
+					break;
+				}
+			}
+
+			if(bAcceptTriangle)
+			{
+				const FMatrix44f& RefMatrix = Mesh->GetRefBasesInvMatrix()[ReferenceBoneIndex];
+
+				const int32 ResultBaseIndex = Result.Positions.Num();
+				Result.Triangles.Add({ ResultBaseIndex, ResultBaseIndex + 1, ResultBaseIndex + 2 });
 
-				if(RefSkeleton.GetBoneName(MeshBoneIndex) == BoneName)
+				for(int32 TriVertIndex = 0; TriVertIndex < 3; ++TriVertIndex)
 				{
-					const FMatrix44f& RefMatrix = Mesh->GetRefBasesInvMatrix()[MeshBoneIndex];
+					const uint32 VertIndex = TriBase[TriVertIndex];
 
 					Result.Positions.Add(RefMatrix.TransformPosition(Positions.VertexPosition(VertIndex)));
 					Result.Normals.Add(RefMatrix.TransformVector(MeshVertexBuffer.VertexTangentZ(VertIndex)));
-					Result.Weights.Add(MaxInfluence * WeightsToFloatMult);
+
+					// Update bone bounds
 					BoundsBuilder += Result.Positions.Last();
+
+					const int32 RefBoneIndex = Section.BoneMap[Weights->GetBoneIndex(VertIndex, ReferenceInfluenceIndex)];
+					const FName RefBoneName = RefSkeleton.GetBoneName(RefBoneIndex);
+
+					// Presume zero (if we don't find the bone then it has a zero weight). The checks below are duplicated because
+					// this first check should pass for the majority of the verts in the selection - only the boundary verts could
+					// potentially require a search through the list of influences, so we quickly check the expected influcence before
+					// falling back to a search
+					Result.Weights.Add(0);
+					if(RefBoneName == BoneName)
+					{
+						// Influence matches, take weight directly
+						Result.Weights.Last() = Weights->GetBoneWeight(VertIndex, ReferenceInfluenceIndex) * WeightsToFloatMult;
+					}
+					else
+					{
+						// Need to see if it's on a different influence
+						const int32 NumInfluences = Weights->GetMaxBoneInfluences();
+						for(int32 InfluenceIndex = 0; InfluenceIndex < NumInfluences; ++InfluenceIndex)
+						{
+							const int32 SectionBoneIndex = Weights->GetBoneIndex(VertIndex, InfluenceIndex);
+
+							if(!Section.BoneMap.IsValidIndex(SectionBoneIndex))
+							{
+								// Invalid bone index for this influence
+								continue;
+							}
+
+							const int32 TestBoneIndex = Section.BoneMap[Weights->GetBoneIndex(VertIndex, InfluenceIndex)];
+							const FName TestBoneName = RefSkeleton.GetBoneName(TestBoneIndex);
+
+							if(TestBoneName == BoneName)
+							{
+								Result.Weights.Last() = Weights->GetBoneWeight(VertIndex, InfluenceIndex) * WeightsToFloatMult;
+							}
+						}
+					}
 				}
 			}
 		}
@@ -256,6 +351,10 @@ FBoneVertData GetMergedBoneVertexData(const FRigidAssetBoneInfo& Bone, USkeletal
 			return {};
 		}
 
+		// Track the number of verts before appending so we can rebase the child triangles
+		const int32 VertBase = RootData.Positions.Num();
+		const int32 TriangleBase = RootData.Triangles.Num();
+
 		RootData.Positions.Reserve(RootData.Positions.Num() + ChildData.Positions.Num());
 		RootData.Normals.Reserve(RootData.Normals.Num() + ChildData.Normals.Num());
 		for(int32 VertIndex = 0; VertIndex < ChildData.Positions.Num(); ++VertIndex)
@@ -267,6 +366,15 @@ FBoneVertData GetMergedBoneVertexData(const FRigidAssetBoneInfo& Bone, USkeletal
 			RootData.Normals.Add(FVector3f(NormalInParentFrame));
 		}
 
+		RootData.Triangles.Append(ChildData.Triangles);
+
+		for(int32 TriIndex = TriangleBase; TriIndex < RootData.Triangles.Num(); ++TriIndex)
+		{
+			RootData.Triangles[TriIndex][0] += VertBase;
+			RootData.Triangles[TriIndex][1] += VertBase;
+			RootData.Triangles[TriIndex][2] += VertBase;
+		}
+
 		RootData.Weights.Append(ChildData.Weights);
 		RootData.Bounds = RootData.Bounds + ChildData.Bounds.TransformBy(FTransform3f(ToParent));
 	}
@@ -504,7 +612,7 @@ void UBoneGeometryGenerator_Sphere::DebugDraw(UE::Dataflow::FContext& Context, I
 		{
 			continue;
 		}
-	
+
 		FBoneVertData BoneVerts = GetMergedBoneVertexData(Bone, Bones.Mesh, BaseSettings.SourceLod, BaseSettings.VertexMode);
 
 		DataflowRenderingInterface.SetColor(FLinearColor::MakeRandomSeededColor(BoneColorSeed++));
@@ -804,3 +912,134 @@ TArray<FRigidAssetBoneGeometry> UBoneGeometryGenerator_Convex::Build(FRigidAsset
 
 	return Result;
 }
+
+////////////////////////////////////////////////////////////////////////////////////////////////////////////////
+
+TArray<FRigidAssetBoneGeometry> UBoneGeometryGenerator_ConvexDecomposition::Build(FRigidAssetBoneSelection Bones)
+{
+	using namespace UE::Chaos::RigidAsset;
+
+	if(!Bones.Mesh || !Bones.Skeleton)
+	{
+		return {};
+	}
+
+	TArray<FRigidAssetBoneGeometry> Result;
+
+	for(const FRigidAssetBoneInfo& Bone : Bones.SelectedBones)
+	{
+		FBoneVertData BoneVerts = GetMergedBoneVertexData(Bone, Bones.Mesh, BaseSettings.SourceLod, BaseSettings.VertexMode);
+
+		if(BoneVerts.Positions.Num() == 0)
+		{
+			// No influences, make a small sphere and warn
+			// #TODO pass context through for dataflow warnings
+			Result.Emplace(Bone, FSimpleGeometry(MakeShared<FKSphereElem>(5.0f)));
+			continue;
+		}
+
+		UE::Geometry::FConvexDecomposition3 Decomp;
+		UE::Geometry::FConvexDecomposition3::FPreprocessMeshOptions PreprocessOptions;
+
+		// Merge edges by default and thicken up to 1cm if any axis is coplanar (flat mesh)
+		PreprocessOptions.bMergeEdges = true;
+		PreprocessOptions.ThickenInputAfterHullFailure = 1.0f;
+
+		TArrayView<const FVector3f> PositionView(BoneVerts.Positions);
+		TArrayView<const FIntVector3> IndexView(BoneVerts.Triangles);
+		Decomp.InitializeFromIndexMesh(PositionView, IndexView, PreprocessOptions);
+
+		Decomp.bTreatAsSolid = Decomp.IsInputSolid();
+
+		if(Method == EDecompositionMethod::Simple)
+		{
+			Decomp.Compute(FMath::Max(1, NumHulls));
+		}
+		else
+		{
+			UE::Geometry::FNegativeSpaceSampleSettings NegativeSpaceSettings;
+
+			NegativeSpaceSettings.ApplyDefaults();
+			NegativeSpaceSettings.bOnlyConnectedToHull = bOnlyConnectedToHull;
+			NegativeSpaceSettings.MinRadius = NegativeSpaceMinRadius;
+			NegativeSpaceSettings.ReduceRadiusMargin = NegativeSpaceTolerance;
+			NegativeSpaceSettings.bRequireSearchSampleCoverage = true;
+			NegativeSpaceSettings.bAllowSamplesInsideMesh = Decomp.bTreatAsSolid;
+
+			Decomp.MaxConvexEdgePlanes = 4;
+			Decomp.bSplitDisconnectedComponents = false;
+			Decomp.ConvexEdgeAngleMoreSamplesThreshold = 180;
+			Decomp.InitializeNegativeSpace(NegativeSpaceSettings);
+
+			constexpr int32 MaxUnconstrainedSplits = 10000;
+			const int32 NumSplits = NegativeSpaceMaxSplits > 0 ? NegativeSpaceMaxSplits : MaxUnconstrainedSplits;
+
+			for(int32 SplitIndex = 0; SplitIndex < NumSplits; ++SplitIndex)
+			{
+				int32 NumPerformedSplits = Decomp.SplitWorst(false, 0, true, NegativeSpaceTolerance * 0.5);
+
+				if(NumPerformedSplits == 0)
+				{
+					// No more splits possible - decomp is complete
+					break;
+				}
+
+				if(SplitIndex == MaxUnconstrainedSplits - 1)
+				{
+					UE_LOG(LogBoneGeometryGenerators, Warning, TEXT("Decomposition ran maximum number of splits (%d), potentially stuck in a splitting loop"), MaxUnconstrainedSplits);
+				}
+			}
+
+			Decomp.FixHullOverlapsInNegativeSpace();
+
+			UE::Geometry::FConvexDecomposition3::FMergeSettings MergeSettings;
+			MergeSettings.TargetNumParts = NegativeSpaceMaxSplits > 0 ? NegativeSpaceMaxSplits : -1;
+			MergeSettings.MinThicknessTolerance = MinThicknessTolerance;
+			MergeSettings.bAllowCompact = true;
+
+			Decomp.MergeBest(MergeSettings);
+		}
+
+		const int32 NumComputedHulls = Decomp.NumHulls();
+
+		for(int32 HullIndex = 0; HullIndex < NumComputedHulls; ++HullIndex)
+		{
+			TSharedPtr<FKConvexElem> NewConvex = MakeShared<FKConvexElem>();
+
+			if(bSimplifyHulls)
+			{
+				UE::Geometry::FDynamicMesh3 HullMesh = Decomp.GetHullMesh(HullIndex);
+
+				UE::Geometry::FMeshConvexHull::FSimplifyOptions Options{
+					SimplifyTargetMaxFaces,
+					SimplifyGeometricToleranceAfterTarget
+				};
+
+				if(!ensure(UE::Geometry::FMeshConvexHull::SimplifyHull(HullMesh, Options)))
+				{
+					UE_LOG(LogBoneGeometryGenerators, Error, TEXT("Simplification failed for convex hull %d of %d, skipping"), HullIndex, NumComputedHulls);
+					continue;
+				}
+
+				for(int32 MeshVertIndex : HullMesh.VertexIndicesItr())
+				{
+					NewConvex->VertexData.Add(static_cast<FVector>(HullMesh.GetVertex(MeshVertIndex)));
+				}
+			}
+			else
+			{
+				TArray<FVector3f> DecompVerts = Decomp.GetVertices<float>(HullIndex);
+				for(const FVector3f& Vert : DecompVerts)
+				{
+					NewConvex->VertexData.Add(static_cast<FVector>(Vert));
+				}
+			}
+
+			NewConvex->UpdateElemBox();
+
+			Result.Emplace(Bone, FSimpleGeometry(NewConvex));
+		}
+	}
+
+	return Result;
+}
diff --git a/Engine/Plugins/Experimental/ChaosRigidAsset/Source/ChaosRigidAssetNodes/Private/Generators/BoneGeometryGenerators.h b/Engine/Plugins/Experimental/ChaosRigidAsset/Source/ChaosRigidAssetNodes/Private/Generators/BoneGeometryGenerators.h
--- a/Engine/Plugins/Experimental/ChaosRigidAsset/Source/ChaosRigidAssetNodes/Private/Generators/BoneGeometryGenerators.h
+++ b/Engine/Plugins/Experimental/ChaosRigidAsset/Source/ChaosRigidAssetNodes/Private/Generators/BoneGeometryGenerators.h
@@ -3,10 +3,16 @@
 #include "BoneSelection.h"
 #include "RigidDataflowNode.h"
 #include "ShapeElemNodes.h"
+
+#include "Logging/LogCategory.h"
+#include "Logging/LogMacros.h"
+
 #include "BoneGeometryGenerators.generated.h"
 
 #pragma once
 
+DECLARE_LOG_CATEGORY_EXTERN(LogBoneGeometryGenerators, Log, All);
+
 namespace UE::Dataflow
 {
 	void RegisterBoneGeometryGeneratorNodes();
@@ -125,6 +131,7 @@ struct FBoneVertData
 	TArray<FVector3f> Positions;
 	TArray<FVector3f> Normals;
 	TArray<float> Weights;
+	TArray<FIntVector3> Triangles;
 	FBoxSphereBounds3f Bounds;
 };
 
@@ -261,7 +268,7 @@ private:
 
 /**
  * Convex generator
- * Takes the extracted verts then uses them to generate a convex hull decomposition
+ * Takes the extracted verts then uses them to generate a single convex hull
  */
 UCLASS()
 class UBoneGeometryGenerator_Convex : public UBoneGeometryGenerator
@@ -277,6 +284,74 @@ private:
 
 ////////////////////////////////////////////////////////////////////////////////////////////////////////////////
 
+UENUM()
+enum class EDecompositionMethod
+{
+	Simple,
+	NegativeSpace
+};
+
+/**
+ * Convex decomposition generator
+ * Takes the extracted verts then uses them to generate a convex hull decomposition of multiple hulls
+ */
+UCLASS()
+class UBoneGeometryGenerator_ConvexDecomposition : public UBoneGeometryGenerator
+{
+	GENERATED_BODY()
+
+public:
+
+	UPROPERTY(EditAnywhere, Category = Geometry)
+	EDecompositionMethod Method = EDecompositionMethod::Simple;
+
+	/** Number of hulls to create in the Simple decomposition mode */
+	UPROPERTY(EditAnywhere, Category = Geometry, meta = (EditCondition="Method == EDecompositionMethod::Simple", EditConditionHides))
+	int32 NumHulls = 4;
+
+	/** When protecting negative space, only look for space that is connected out to the convex hull. This removes inaccessable internal negative space from consideration. */
+	UPROPERTY(EditAnywhere, Category = Geometry, meta = (EditCondition="Method == EDecompositionMethod::NegativeSpace", EditConditionHides))
+	bool bOnlyConnectedToHull = true;
+
+	/** Amount of space to leave between convex hulls and protected negative space */
+	UPROPERTY(EditAnywhere, Category = Geometry, meta = (EditCondition="Method == EDecompositionMethod::NegativeSpace", EditConditionHides))
+	float NegativeSpaceTolerance = 2.0f;
+
+	/** Spheres smaller than this are not included in the negative space */
+	UPROPERTY(EditAnywhere, Category = Geometry, meta = (EditCondition="Method == EDecompositionMethod::NegativeSpace", EditConditionHides))
+	float NegativeSpaceMinRadius = 10.0f;
+
+	/** 
+	 * Number of splits to target in the final decomposition when using the negative space method.
+	 * If 0 the number of splits will be driven by the decomposition and may produce a more accurate result but take longer.
+	 * If >0 the number of splits will be truncated at the specified value
+	 */
+	UPROPERTY(EditAnywhere, Category = Geometry, meta = (EditCondition="Method == EDecompositionMethod::NegativeSpace", EditConditionHides))
+	int32 NegativeSpaceMaxSplits = 0;
+
+	/** Optional tolerance (in cm) for merging after decomposition. Hulls under this size will always be merged */
+	UPROPERTY(EditAnywhere, Category = Geometry, meta = (EditCondition="Method == EDecompositionMethod::NegativeSpace", EditConditionHides))
+	double MinThicknessTolerance = 0;
+
+	/** Whether to simplify the resulting hull meshes */
+	UPROPERTY(EditAnywhere, Category = Geometry)
+	bool bSimplifyHulls = false;
+
+	/** The number of faces to target during simplification */
+	UPROPERTY(EditAnywhere, Category = Geometry, meta = (EditCondition="bSimplifyHulls", EditConditionHides))
+	int32 SimplifyTargetMaxFaces = 128;
+
+	/** If >0 the mesh will be further simplified beloc SimplifyTargetMaxFaces if the mesh error can be kept below this limit */
+	UPROPERTY(EditAnywhere, Category = Geometry, meta = (EditCondition="bSimplifyHulls", EditConditionHides))
+	double SimplifyGeometricToleranceAfterTarget = 0;
+
+	TArray<FRigidAssetBoneGeometry> Build(FRigidAssetBoneSelection Bones) final override;
+
+private:
+};
+
+////////////////////////////////////////////////////////////////////////////////////////////////////////////////
+
 /**
  * Make nodes for each generator, each creates a configured generator that can be wired into the generation node
  */
@@ -379,3 +454,28 @@ private:
 };
 
 ////////////////////////////////////////////////////////////////////////////////////////////////////////////////
+
+USTRUCT()
+struct FMakeConvexDecompBoneGeometryGenerator : public FRigidDataflowNode
+{
+	GENERATED_BODY()
+	DATAFLOW_NODE_DEFINE_INTERNAL(FMakeConvexDecompBoneGeometryGenerator, "Make Convex Decomposition Bone Geometry Builder", "PhysicsAsset", "")
+
+public:
+
+	FMakeConvexDecompBoneGeometryGenerator(const UE::Dataflow::FNodeParameters& InParam, FGuid InGuid = FGuid::NewGuid())
+		: FRigidDataflowNode(InParam, InGuid)
+	{
+		Register();
+	}
+
+private:
+
+	UPROPERTY(EditAnywhere, Instanced, Category = "Generator", meta = (DataflowOutput))
+	TObjectPtr<UBoneGeometryGenerator> Generator;
+
+	void Evaluate(UE::Dataflow::FContext& Context, const FDataflowOutput* Out) const override;
+	void Register();
+};
+
+////////////////////////////////////////////////////////////////////////////////////////////////////////////////
diff --git a/Engine/Plugins/Experimental/ChaosRigidAsset/Source/ChaosRigidAssetNodes/Private/PhysicsAssetDataflowNodes.cpp b/Engine/Plugins/Experimental/ChaosRigidAsset/Source/ChaosRigidAssetNodes/Private/PhysicsAssetDataflowNodes.cpp
--- a/Engine/Plugins/Experimental/ChaosRigidAsset/Source/ChaosRigidAssetNodes/Private/PhysicsAssetDataflowNodes.cpp
+++ b/Engine/Plugins/Experimental/ChaosRigidAsset/Source/ChaosRigidAssetNodes/Private/PhysicsAssetDataflowNodes.cpp
@@ -46,7 +46,7 @@ void FDataflowPhysicsAssetTerminalNode::SetAssetValue(TObjectPtr<UObject> Asset,
 		const FPhysicsAssetDataflowState InState = GetValue(Context, &State);
 		InState.DebugLog();
 
-		UE::Chaos::RigidAsset::FPhysicsAssetBuilder::Make(InState.TargetSkeleton, InState.Bodies.ToArray(), InState.Constraints.ToArray())
+		UE::Chaos::RigidAsset::FPhysicsAssetBuilder::Make(InState.GetSkeleton(), InState.GetBodies().ToArray(), InState.Constraints.ToArray())
 			.SetTargetAsset(Target)
 			.Build();
 	}
@@ -104,13 +104,13 @@ void FDataflowPhysicsAssetAddBody::Evaluate(UE::Dataflow::FContext& Context, con
 {
 	if(Out->IsA(&State))
 	{
-		const USkeletalBodySetup* InBody = GetValue(Context, &Body);
+		USkeletalBodySetup* InBody = GetValue(Context, &Body);
 
 		FPhysicsAssetDataflowState InState = GetValue(Context, &State);
 
-		if(InBody)
+		if(InBody && InBody->BoneName != NAME_None)
 		{
-			InState.Bodies.Emplace(Cast<USkeletalBodySetup>(StaticDuplicateObject(InBody, GetTransientPackage())));
+			InState.SetBody(InBody->BoneName, InBody);
 		}
 
 		SetValue(Context, InState, &State);
@@ -361,8 +361,8 @@ void FDataflowSelectBonesByName::Evaluate(UE::Dataflow::FContext& Context, const
 	{
 		FPhysicsAssetDataflowState InState = GetValue(Context, &State);
 
-		TObjectPtr<USkeleton> Skeleton = InState.TargetSkeleton;
-		TObjectPtr<USkeletalMesh> Mesh = InState.TargetMesh;
+		TObjectPtr<USkeleton> Skeleton = InState.GetSkeleton();
+		TObjectPtr<USkeletalMesh> Mesh = InState.GetMesh();
 
 		if(!Skeleton)
 		{
@@ -598,22 +598,31 @@ void FDataflowCreateGeometryForBones::Evaluate(UE::Dataflow::FContext& Context,
 		}
 
 		FPhysicsAssetDataflowState InState = GetValue(Context, &State);
-		TArray<FRigidAssetBoneGeometry> Geoms = InGenerator->Build(InSelection);
 
-		if(Geoms.Num() != InSelection.SelectedBones.Num())
+		if(!InState.HasData())
 		{
-			Context.Error(TEXT("Failed to generate bodies for each bone in the selection"));
+			Context.Error(TEXT("Attempted to make a body without a valid asset state"));
 			SafeForwardInput(Context, &State, &State);
 			return;
 		}
 
+		// Invoke the geometry generator
+		TArray<FRigidAssetBoneGeometry> Geoms = InGenerator->Build(InSelection);
+
+		// Create new bodies for the geometry - or append new geometries to existing bodies if they already exist
 		for(int32 Index = 0; Index < Geoms.Num(); ++Index)
 		{
 			const UE::Chaos::RigidAsset::FSimpleGeometry& Geom = Geoms[Index].Geometry;
 			const FRigidAssetBoneInfo& Bone = Geoms[Index].Bone;
 
-			TObjectPtr<USkeletalBodySetup> NewSetup = CastChecked<USkeletalBodySetup>(StaticDuplicateObject(InTemplate, GetTransientPackage()));
+			TObjectPtr<USkeletalBodySetup> TemplateToUse = InTemplate;
+			if(!TemplateToUse)
+			{
+				TemplateToUse = USkeletalBodySetup::StaticClass()->GetDefaultObject<USkeletalBodySetup>();
+			}
 
+			USkeletalBodySetup* BoneSetup = InState.FindOrCreateBodyChecked(Bone.Name, TemplateToUse);
+			
 			if(!Geom.GeometryElement)
 			{
 				Context.Error(TEXT("Generator failed to create usable geometry"));
@@ -621,14 +630,10 @@ void FDataflowCreateGeometryForBones::Evaluate(UE::Dataflow::FContext& Context,
 				return;
 			}
 
-			NewSetup->AggGeom.VisitShapeAndContainer(*Geom.GeometryElement.Get(), [] <typename ElemType> (const ElemType & ElemTyped, TArray<ElemType>&Container)
+			BoneSetup->AggGeom.VisitShapeAndContainer(*Geom.GeometryElement.Get(), [] <typename ElemType> (const ElemType & ElemTyped, TArray<ElemType>&Container)
 			{
 				Container.Add(ElemTyped);
 			});
-
-			NewSetup->BoneName = Bone.Name;
-
-			InState.Bodies.Add(NewSetup);
 		}
 
 		SetValue(Context, MoveTemp(InState), &State);
@@ -909,9 +914,9 @@ void FDataflowRigidSimulationCachingNode::Evaluate(UE::Dataflow::FContext& Conte
 		
 		TObjectPtr<UAnimSequence> OutAnim = NewObject<UAnimSequence>();
 #if WITH_EDITOR
-		if (InState.TargetSkeleton && InState.TargetMesh)
+		if (InState.HasData())
 		{
-			OutAnim->SetSkeleton(InState.TargetSkeleton);
+			OutAnim->SetSkeleton(InState.GetSkeleton());
 			OutAnim->BoneCompressionSettings = FAnimationUtils::GetDefaultAnimationRecorderBoneCompressionSettings();
 
 			IAnimationDataController& Controller = OutAnim->GetController();
@@ -928,7 +933,7 @@ void FDataflowRigidSimulationCachingNode::Evaluate(UE::Dataflow::FContext& Conte
 			Controller.SetFrameRate(RecordingRate, /*bTransactRecording*/false);
 			Controller.SetNumberOfFrames(NumFrames, /*bTransactRecording*/false);
 
-			OutAnim->RetargetSource = InState.TargetSkeleton->GetRetargetSourceForMesh(InState.TargetMesh);
+			OutAnim->RetargetSource = InState.GetSkeleton()->GetRetargetSourceForMesh(InState.GetMesh());
 
 			FSimulationState SimState;
 			SimState.FixedTimeStep = FixedTimeStep;
@@ -938,7 +943,7 @@ void FDataflowRigidSimulationCachingNode::Evaluate(UE::Dataflow::FContext& Conte
 				TMap<FName, FRawAnimSequenceTrack> RawTracksByName;
 				TMap<FName, FActorHandle*> ParentActorByChildName;
 
-				const FReferenceSkeleton& RefSkeleton = InState.TargetMesh->GetRefSkeleton();
+				const FReferenceSkeleton& RefSkeleton = InState.GetMesh()->GetRefSkeleton();
 
 				// create the tracks from the simulation actors
 				for (const TPair<FName, FActorHandle*>& ActorByName : SimState.ActorsByName)
@@ -1028,7 +1033,7 @@ void FDataflowRigidSimulationCachingNode::BuildSimulationFromState(const FPhysic
 {
 	using namespace ImmediatePhysics;
 
-	if (InState.TargetMesh == nullptr)
+	if (InState.GetMesh() == nullptr)
 	{
 		return;
 	}
@@ -1042,13 +1047,13 @@ void FDataflowRigidSimulationCachingNode::BuildSimulationFromState(const FPhysic
 	// structure to efficiently compute bones component space
 	TArray<FTransform> CachedTransforms;
 	TArray<bool> CachedTransformReady;
-	const FReferenceSkeleton& RefSkeleton = InState.TargetMesh->GetRefSkeleton();
+	const FReferenceSkeleton& RefSkeleton = InState.GetMesh()->GetRefSkeleton();
 	CachedTransforms.SetNumUninitialized(RefSkeleton.GetRefBonePose().Num());
 	CachedTransformReady.SetNumZeroed(RefSkeleton.GetRefBonePose().Num());
 
-	for (TObjectPtr<USkeletalBodySetup> BodySetup : InState.Bodies)
+	for (TObjectPtr<USkeletalBodySetup> BodySetup : InState.GetBodies())
 	{
-		TUniquePtr<FBodyInstance> BodyInstance = CreateBodyInstance(BodySetup);
+		TUniquePtr<FBodyInstance> BodyInstance = CreateBodyInstance(BodySetup.Get());
 		if (BodyInstance)
 		{
 			const int32 BoneIndex = RefSkeleton.FindBoneIndex(BodySetup->BoneName);
diff --git a/Engine/Plugins/Experimental/ChaosRigidAsset/Source/ChaosRigidAssetNodes/Private/PhysicsAssetDataflowState.cpp b/Engine/Plugins/Experimental/ChaosRigidAsset/Source/ChaosRigidAssetNodes/Private/PhysicsAssetDataflowState.cpp
--- a/Engine/Plugins/Experimental/ChaosRigidAsset/Source/ChaosRigidAssetNodes/Private/PhysicsAssetDataflowState.cpp
+++ b/Engine/Plugins/Experimental/ChaosRigidAsset/Source/ChaosRigidAssetNodes/Private/PhysicsAssetDataflowState.cpp
@@ -3,6 +3,141 @@
 #include "PhysicsAssetDataflowState.h"
 #include "PhysicsEngine/SkeletalBodySetup.h"
 #include "PhysicsEngine/PhysicsConstraintTemplate.h"
+#include "Animation/Skeleton.h"
+#include "UObject/Package.h"
+
+FPhysicsAssetDataflowState::FPhysicsAssetDataflowState() = default;
+
+FPhysicsAssetDataflowState::FPhysicsAssetDataflowState(TObjectPtr<USkeleton> InTargetSkeleton, TObjectPtr<USkeletalMesh> InTargetMesh)
+	: TargetSkeleton(InTargetSkeleton)
+	, TargetMesh(InTargetMesh)
+{
+	if(HasData())
+	{
+		const int32 NumBones = TargetSkeleton->GetReferenceSkeleton().GetNum();
+		BoneIndexToBody.SetNum(NumBones);
+
+		for(int32 Index = 0; Index < NumBones; ++Index)
+		{
+			BoneIndexToBody.Modify(Index, INDEX_NONE);
+		}
+	}
+}
+
+USkeleton* FPhysicsAssetDataflowState::GetSkeleton() const
+{
+	return TargetSkeleton.Get();
+}
+
+USkeletalMesh* FPhysicsAssetDataflowState::GetMesh() const
+{
+	return TargetMesh.Get();
+}
+
+USkeletalBodySetup* FPhysicsAssetDataflowState::GetBody(FName BoneName) const
+{
+	if(!HasData())
+	{
+		return nullptr;
+	}
+
+	const FReferenceSkeleton& RefSkel = TargetSkeleton->GetReferenceSkeleton();
+
+	const int32 BoneIndex = RefSkel.FindBoneIndex(BoneName);
+
+	if(BoneIndex == INDEX_NONE)
+	{
+		return nullptr;
+	}
+
+	int32 BodyIndex = BoneIndexToBody[BoneIndex];
+
+	if(BodyIndex == INDEX_NONE)
+	{
+		return nullptr;
+	}
+
+	return Bodies[BodyIndex];
+}
+
+void FPhysicsAssetDataflowState::SetBody(FName BoneName, USkeletalBodySetup* InSetup)
+{
+	if(!HasData())
+	{
+		return;
+	}
+
+	const FReferenceSkeleton& RefSkel = TargetSkeleton->GetReferenceSkeleton();
+
+	const int32 BoneIndex = RefSkel.FindBoneIndex(BoneName);
+
+	if(BoneIndex == INDEX_NONE)
+	{
+		return;
+	}
+
+	int32 BodyIndex = BoneIndexToBody[BoneIndex];
+
+	if(BodyIndex == INDEX_NONE)
+	{
+		Bodies.Add(InSetup);
+		BoneIndexToBody.Modify(BoneIndex, Bodies.Num() - 1);
+	}
+	else
+	{
+		Bodies.Modify(BodyIndex, InSetup);
+	}
+}
+
+USkeletalBodySetup* FPhysicsAssetDataflowState::FindOrCreateBody(FName BoneName, USkeletalBodySetup* TemplateBody)
+{
+	if(!HasData())
+	{
+		return nullptr;
+	}
+
+	const FReferenceSkeleton& RefSkel = TargetSkeleton->GetReferenceSkeleton();
+
+	const int32 BoneIndex = RefSkel.FindBoneIndex(BoneName);
+
+	if(BoneIndex == INDEX_NONE)
+	{
+		return nullptr;
+	}
+
+	int32 BodyIndex = BoneIndexToBody[BoneIndex];
+
+	if(BodyIndex == INDEX_NONE)
+	{
+		if(!TemplateBody)
+		{
+			return nullptr;
+		}
+
+		USkeletalBodySetup* NewSetup = CastChecked<USkeletalBodySetup>(StaticDuplicateObject(TemplateBody, GetTransientPackage()));
+		NewSetup->BoneName = BoneName;
+
+		Bodies.Add(NewSetup);
+		BoneIndexToBody.Modify(BoneIndex, Bodies.Num() - 1);
+
+		return Bodies.Last();
+	}
+
+	return Bodies[BodyIndex];
+}
+
+USkeletalBodySetup* FPhysicsAssetDataflowState::FindOrCreateBodyChecked(FName BoneName, USkeletalBodySetup* TemplateBody /*= nullptr*/)
+{
+	USkeletalBodySetup* Body = FindOrCreateBody(BoneName, TemplateBody);
+	check(Body);
+
+	return Body;
+}
+
+const TCopyOnWriteArray<TObjectPtr<USkeletalBodySetup>>& FPhysicsAssetDataflowState::GetBodies() const
+{
+	return Bodies;
+}
 
 void FPhysicsAssetDataflowState::AddStructReferencedObjects(FReferenceCollector& Collector)
 {
@@ -15,11 +150,21 @@ void FPhysicsAssetDataflowState::DebugLog() const
 {
 	for(const USkeletalBodySetup* Setup : Bodies)
 	{
+		if(!Setup)
+		{
+			continue;
+		}
+
 		UE_LOG(LogTemp, Warning, TEXT("Body: %s"), *Setup->GetName());
 	}
 
 	for(const UPhysicsConstraintTemplate* Constraint : Constraints)
 	{
+		if(!Constraint)
+		{
+			continue;
+		}
+
 		UE_LOG(LogTemp, Warning, TEXT("Constraint: %s"), *Constraint->GetName());
 	}
 }
diff --git a/Engine/Plugins/Experimental/ChaosRigidAsset/Source/ChaosRigidAssetNodes/Public/CopyOnWriteArray.h b/Engine/Plugins/Experimental/ChaosRigidAsset/Source/ChaosRigidAssetNodes/Public/CopyOnWriteArray.h
--- a/Engine/Plugins/Experimental/ChaosRigidAsset/Source/ChaosRigidAssetNodes/Public/CopyOnWriteArray.h
+++ b/Engine/Plugins/Experimental/ChaosRigidAsset/Source/ChaosRigidAssetNodes/Public/CopyOnWriteArray.h
@@ -47,6 +47,9 @@ template<typename T>
 class TCopyOnWriteArray
 {
 public:
+	using InnerArrayType = TArray<T>;
+	using AllocatorType = InnerArrayType::AllocatorType;
+	using SizeType = InnerArrayType::SizeType;
 
 	TCopyOnWriteArray()
 	{
@@ -75,7 +78,7 @@ public:
 		return *this;
 	}
 
-	const T& operator [](int32 Index)
+	const T& operator [](int32 Index) const
 	{
 		checkf(State, TEXT("Array index out of bounds: %d into an array of size 0"), Index);
 		return State->Arr[Index];
@@ -160,6 +163,33 @@ public:
 		State->Arr.Append(Other);
 	}
 
+	void SetNum(SizeType NewNum, EAllowShrinking AllowShrink = UE::Core::Private::AllowShrinkingByDefault<AllocatorType>())
+	{
+		FWriteScopeLock ScopeLock(Lock);
+
+		CheckCopyOnWrite();
+		State->Arr.SetNum(NewNum, AllowShrink);
+	}
+
+	void SetNumZeroed(SizeType NewNum, EAllowShrinking AllowShrink = UE::Core::Private::AllowShrinkingByDefault<AllocatorType>())
+	{
+		FWriteScopeLock ScopeLock(Lock);
+
+		CheckCopyOnWrite();
+		State->Arr.SetNumZeroed(NewNum, AllowShrink);
+	}
+
+	SizeType Num() const
+	{
+		return State ? State->Arr.Num() : 0;
+	}
+
+	const T& Last(SizeType IndexFromEnd = 0) const
+	{
+		checkf(State, TEXT("Array index out of bounds: Last() accessing empty array"));
+		return State->Arr.Last(IndexFromEnd);
+	}
+
 	// If the array contains UObjects, we need to emit references to them directly
 	void AddReferencedObjects(FReferenceCollector& Collector) const
 		requires CIsUObjectType<T>
