[Dataflow] [RigidAsset] Refactored selection merging and geometry building for decomp to be parallel, one task per-bone.

#Synced-CL 50199821

diff --git a/Engine/Plugins/Experimental/ChaosRigidAsset/Source/ChaosRigidAssetNodes/Private/Generators/BoneGeometryGenerators.cpp b/Engine/Plugins/Experimental/ChaosRigidAsset/Source/ChaosRigidAssetNodes/Private/Generators/BoneGeometryGenerators.cpp
--- a/Engine/Plugins/Experimental/ChaosRigidAsset/Source/ChaosRigidAssetNodes/Private/Generators/BoneGeometryGenerators.cpp
+++ b/Engine/Plugins/Experimental/ChaosRigidAsset/Source/ChaosRigidAssetNodes/Private/Generators/BoneGeometryGenerators.cpp
@@ -20,6 +20,45 @@
 
 DEFINE_LOG_CATEGORY(LogBoneGeometryGenerators);
 
+////////////////////////////////////////////////////////////////////////////////////////////////////////////////
+
+static bool GBoneGeometryGeneratorsRunParallel = true;
+static FAutoConsoleVariableRef CvarBoneGeometryGeneratorsRunParallel(
+	TEXT("p.dataflow.bonegeometrygenerator.parallel"),
+	GBoneGeometryGeneratorsRunParallel,
+	TEXT("Whether bone geometry generator algorithms can run in parallel"));
+
+template<int32 SerialThreshold = 4, typename Callable>
+TArray<UE::Tasks::FTask> DispatchTasks(const TCHAR* InDebugName, int32 NumTasks, Callable&& InCallable)
+{
+	const bool bRunParallel = NumTasks > SerialThreshold && GBoneGeometryGeneratorsRunParallel && FApp::ShouldUseThreadingForPerformance();
+
+	if(!bRunParallel)
+	{
+		for(int32 TaskIndex = 0; TaskIndex < NumTasks; ++TaskIndex)
+		{
+			InCallable(TaskIndex);
+		}
+
+		// No tasks returned, didn't go wide
+		return {};
+	}
+	else
+	{
+		TArray<UE::Tasks::FTask> Tasks;
+		Tasks.AddDefaulted(NumTasks);
+
+		for(int32 TaskIndex = 0; TaskIndex < NumTasks; ++TaskIndex)
+		{
+			Tasks[TaskIndex] = UE::Tasks::Launch(InDebugName, [TaskIndex, &InCallable]() { InCallable(TaskIndex); }, LowLevelTasks::ETaskPriority::High);
+		}
+
+		return Tasks;
+	}
+}
+
+////////////////////////////////////////////////////////////////////////////////////////////////////////////////
+
 namespace UE::Dataflow
 {
 	void RegisterBoneGeometryGeneratorNodes()
@@ -158,8 +197,6 @@ FBoneVertData GetBoneVertexData(FName BoneName, USkeletalMesh* Mesh, int32 Lod,
 
 	FBoneVertData Result;
 
-	// Make sure we have initialised the mesh in order to get vertex data
-	Mesh->CalculateInvRefMatrices();
 	const FReferenceSkeleton& RefSkeleton = Mesh->GetRefSkeleton();
 	const FSkeletalMeshLODRenderData& LodData = MeshRenderData->LODRenderData[Lod];
 
@@ -171,11 +208,10 @@ FBoneVertData GetBoneVertexData(FName BoneName, USkeletalMesh* Mesh, int32 Lod,
 	const FStaticMeshVertexBuffer& MeshVertexBuffer = LodData.StaticVertexBuffers.StaticMeshVertexBuffer;
 	const FPositionVertexBuffer& Positions = LodData.StaticVertexBuffers.PositionVertexBuffer;
 
-	TArray<uint32> Indices;
-	LodData.MultiSizeIndexContainer.GetIndexBuffer(Indices);
+	const FRawStaticIndexBuffer16or32Interface* IndexInterface = LodData.MultiSizeIndexContainer.GetIndexBuffer();
 
-	check(Indices.Num() % 3 == 0);
-	const int32 NumTris = Indices.Num() / 3;
+	check(IndexInterface->Num() % 3 == 0);
+	const int32 NumTris = IndexInterface->Num() / 3;
 
 	const uint32 NumLodVerts = Weights->GetNumVertices();
 	const uint32 MaxInfluences = LodData.GetVertexBufferMaxBoneInfluences();
@@ -195,7 +231,7 @@ FBoneVertData GetBoneVertexData(FName BoneName, USkeletalMesh* Mesh, int32 Lod,
 
 		for(uint32 TriIndex = 0; TriIndex < Section.NumTriangles; ++TriIndex)
 		{
-			const uint32* TriBase = &Indices[Section.BaseIndex + TriIndex * 3];
+			const uint32 TriBaseIndex = Section.BaseIndex + TriIndex * 3;
 
 			bool bAcceptTriangle = false;
 			uint32 ReferenceBoneIndex = 0;
@@ -203,7 +239,7 @@ FBoneVertData GetBoneVertexData(FName BoneName, USkeletalMesh* Mesh, int32 Lod,
 
 			for(int32 TriVertIndex = 0; TriVertIndex < 3; ++TriVertIndex)
 			{
-				const uint32 VertIndex = TriBase[TriVertIndex];
+				const uint32 VertIndex = IndexInterface->Get(TriBaseIndex + TriVertIndex);
 
 				if(SelectMode == EVertexSelectMode::Any)
 				{
@@ -267,7 +303,7 @@ FBoneVertData GetBoneVertexData(FName BoneName, USkeletalMesh* Mesh, int32 Lod,
 
 				for(int32 TriVertIndex = 0; TriVertIndex < 3; ++TriVertIndex)
 				{
-					const uint32 VertIndex = TriBase[TriVertIndex];
+					const uint32 VertIndex = IndexInterface->Get(TriBaseIndex + TriVertIndex);
 
 					Result.Positions.Add(RefMatrix.TransformPosition(Positions.VertexPosition(VertIndex)));
 					Result.Normals.Add(RefMatrix.TransformVector(MeshVertexBuffer.VertexTangentZ(VertIndex)));
@@ -409,17 +445,20 @@ FRigidAssetBoneSelection MergeSelection(const FBaseGenerationSettings& Settings,
 
 	const int32 NumSelectionBones = InSelectionSorted.SelectedBones.Num();
 	TArray<FLocalBoneData> BoneData;
-	BoneData.Reserve(NumSelectionBones);
+	BoneData.AddDefaulted(NumSelectionBones);
 
-	for(const FRigidAssetBoneInfo& Bone : InSelectionSorted.SelectedBones)
-	{
-		BoneData.AddDefaulted();
-		FLocalBoneData& Local = BoneData.Last();
+	constexpr static int32 SerialThreshold_GetBoneVertexData = 4;
+	TArray<UE::Tasks::FTask> WorkerTasks = DispatchTasks<SerialThreshold_GetBoneVertexData>(UE_SOURCE_LOCATION, NumSelectionBones, [Mesh, &BoneData, &InSelectionSorted, &Settings](int32 Index)
+		{
+			FLocalBoneData& Local = BoneData[Index];
+			const FRigidAssetBoneInfo& Bone = InSelectionSorted.SelectedBones[Index];
 
-		Local.Bone = Bone;
-		Local.RootBone = Bone;
-		Local.VertData = GetBoneVertexData(Bone.Name, Mesh, Settings.SourceLod, Settings.VertexMode);
-	}
+			Local.Bone = Bone;
+			Local.RootBone = Bone;
+			Local.VertData = GetBoneVertexData(Bone.Name, Mesh, Settings.SourceLod, Settings.VertexMode);
+		});
+
+	UE::Tasks::Wait(WorkerTasks);
 
 	const float MinSizeSq = Settings.Operation == EMergeOperation::MergeSmall ? Settings.MinimumBoneSize * Settings.MinimumBoneSize : std::numeric_limits<float>::max();
 	const float MergeThresholdSizeSq = Settings.Operation == EMergeOperation::MergeSmall ? Settings.SmallBoneMergeThresholdOverride * Settings.SmallBoneMergeThresholdOverride : 0;
@@ -915,6 +954,11 @@ TArray<FRigidAssetBoneGeometry> UBoneGeometryGenerator_Convex::Build(FRigidAsset
 
 ////////////////////////////////////////////////////////////////////////////////////////////////////////////////
 
+// Wrapper struct to pad out T to cache-line size to avoid false sharing when we set up an output
+// array for a multi-threaded algorithm, avoids multiple threads ever reading/writing the same cache line
+template<typename T>
+struct alignas(PLATFORM_CACHE_LINE_SIZE) TPadded : public T {};
+
 TArray<FRigidAssetBoneGeometry> UBoneGeometryGenerator_ConvexDecomposition::Build(FRigidAssetBoneSelection Bones)
 {
 	using namespace UE::Chaos::RigidAsset;
@@ -924,121 +968,141 @@ TArray<FRigidAssetBoneGeometry> UBoneGeometryGenerator_ConvexDecomposition::Buil
 		return {};
 	}
 
-	TArray<FRigidAssetBoneGeometry> Result;
+	const int32 NumBones = Bones.SelectedBones.Num();
 
-	for(const FRigidAssetBoneInfo& Bone : Bones.SelectedBones)
-	{
-		FBoneVertData BoneVerts = GetMergedBoneVertexData(Bone, Bones.Mesh, BaseSettings.SourceLod, BaseSettings.VertexMode);
+	TArray<TPadded<TArray<FRigidAssetBoneGeometry>>> WorkerResults;
+	WorkerResults.AddDefaulted(NumBones);
 
-		if(BoneVerts.Positions.Num() == 0)
+	constexpr static int32 SerialThreshold_BoneDecompose = 4;
+	TArray<UE::Tasks::FTask> DecompTasks = DispatchTasks<SerialThreshold_BoneDecompose>(UE_SOURCE_LOCATION, NumBones, [this, &Bones, &WorkerResults](int32 BoneIndex)
 		{
-			// No influences, make a small sphere and warn
-			// #TODO pass context through for dataflow warnings
-			Result.Emplace(Bone, FSimpleGeometry(MakeShared<FKSphereElem>(5.0f)));
-			continue;
-		}
-
-		UE::Geometry::FConvexDecomposition3 Decomp;
-		UE::Geometry::FConvexDecomposition3::FPreprocessMeshOptions PreprocessOptions;
+			const FRigidAssetBoneInfo& Bone = Bones.SelectedBones[BoneIndex];
 
-		// Merge edges by default and thicken up to 1cm if any axis is coplanar (flat mesh)
-		PreprocessOptions.bMergeEdges = true;
-		PreprocessOptions.ThickenInputAfterHullFailure = 1.0f;
-
-		TArrayView<const FVector3f> PositionView(BoneVerts.Positions);
-		TArrayView<const FIntVector3> IndexView(BoneVerts.Triangles);
-		Decomp.InitializeFromIndexMesh(PositionView, IndexView, PreprocessOptions);
-
-		Decomp.bTreatAsSolid = Decomp.IsInputSolid();
+			{
+				FBoneVertData BoneVerts = GetMergedBoneVertexData(Bone, Bones.Mesh, BaseSettings.SourceLod, BaseSettings.VertexMode);
 
-		if(Method == EDecompositionMethod::Simple)
-		{
-			Decomp.Compute(FMath::Max(1, NumHulls));
-		}
-		else
-		{
-			UE::Geometry::FNegativeSpaceSampleSettings NegativeSpaceSettings;
+				if(BoneVerts.Positions.Num() == 0)
+				{
+					// No influences, make a small sphere and warn
+					// #TODO pass context through for dataflow warnings
+					WorkerResults[BoneIndex].Emplace(Bone, FSimpleGeometry(MakeShared<FKSphereElem>(5.0f)));
+					return;
+				}
 
-			NegativeSpaceSettings.ApplyDefaults();
-			NegativeSpaceSettings.bOnlyConnectedToHull = bOnlyConnectedToHull;
-			NegativeSpaceSettings.MinRadius = NegativeSpaceMinRadius;
-			NegativeSpaceSettings.ReduceRadiusMargin = NegativeSpaceTolerance;
-			NegativeSpaceSettings.bRequireSearchSampleCoverage = true;
-			NegativeSpaceSettings.bAllowSamplesInsideMesh = Decomp.bTreatAsSolid;
+				UE::Geometry::FConvexDecomposition3 Decomp;
+				UE::Geometry::FConvexDecomposition3::FPreprocessMeshOptions PreprocessOptions;
 
-			Decomp.MaxConvexEdgePlanes = 4;
-			Decomp.bSplitDisconnectedComponents = false;
-			Decomp.ConvexEdgeAngleMoreSamplesThreshold = 180;
-			Decomp.InitializeNegativeSpace(NegativeSpaceSettings);
+				// Merge edges by default and thicken up to 1cm if any axis is coplanar (flat mesh)
+				PreprocessOptions.bMergeEdges = true;
+				PreprocessOptions.ThickenInputAfterHullFailure = 1.0f;
 
-			constexpr int32 MaxUnconstrainedSplits = 10000;
-			const int32 NumSplits = NegativeSpaceMaxSplits > 0 ? NegativeSpaceMaxSplits : MaxUnconstrainedSplits;
+				TArrayView<const FVector3f> PositionView(BoneVerts.Positions);
+				TArrayView<const FIntVector3> IndexView(BoneVerts.Triangles);
+				Decomp.InitializeFromIndexMesh(PositionView, IndexView, PreprocessOptions);
 
-			for(int32 SplitIndex = 0; SplitIndex < NumSplits; ++SplitIndex)
-			{
-				int32 NumPerformedSplits = Decomp.SplitWorst(false, 0, true, NegativeSpaceTolerance * 0.5);
+				Decomp.bTreatAsSolid = Decomp.IsInputSolid();
 
-				if(NumPerformedSplits == 0)
+				if(Method == EDecompositionMethod::Simple)
 				{
-					// No more splits possible - decomp is complete
-					break;
+					Decomp.Compute(FMath::Max(1, NumHulls));
 				}
-
-				if(SplitIndex == MaxUnconstrainedSplits - 1)
+				else
 				{
-					UE_LOG(LogBoneGeometryGenerators, Warning, TEXT("Decomposition ran maximum number of splits (%d), potentially stuck in a splitting loop"), MaxUnconstrainedSplits);
-				}
-			}
+					UE::Geometry::FNegativeSpaceSampleSettings NegativeSpaceSettings;
 
-			Decomp.FixHullOverlapsInNegativeSpace();
+					NegativeSpaceSettings.ApplyDefaults();
+					NegativeSpaceSettings.bOnlyConnectedToHull = bOnlyConnectedToHull;
+					NegativeSpaceSettings.MinRadius = NegativeSpaceMinRadius;
+					NegativeSpaceSettings.ReduceRadiusMargin = NegativeSpaceTolerance;
+					NegativeSpaceSettings.bRequireSearchSampleCoverage = true;
+					NegativeSpaceSettings.bAllowSamplesInsideMesh = Decomp.bTreatAsSolid;
 
-			UE::Geometry::FConvexDecomposition3::FMergeSettings MergeSettings;
-			MergeSettings.TargetNumParts = NegativeSpaceMaxSplits > 0 ? NegativeSpaceMaxSplits : -1;
-			MergeSettings.MinThicknessTolerance = MinThicknessTolerance;
-			MergeSettings.bAllowCompact = true;
+					Decomp.MaxConvexEdgePlanes = 4;
+					Decomp.bSplitDisconnectedComponents = false;
+					Decomp.ConvexEdgeAngleMoreSamplesThreshold = 180;
+					Decomp.InitializeNegativeSpace(NegativeSpaceSettings);
 
-			Decomp.MergeBest(MergeSettings);
-		}
+					constexpr int32 MaxUnconstrainedSplits = 10000;
+					const int32 NumSplits = NegativeSpaceMaxSplits > 0 ? NegativeSpaceMaxSplits : MaxUnconstrainedSplits;
 
-		const int32 NumComputedHulls = Decomp.NumHulls();
+					for(int32 SplitIndex = 0; SplitIndex < NumSplits; ++SplitIndex)
+					{
+						int32 NumPerformedSplits = Decomp.SplitWorst(false, 0, true, NegativeSpaceTolerance * 0.5);
 
-		for(int32 HullIndex = 0; HullIndex < NumComputedHulls; ++HullIndex)
-		{
-			TSharedPtr<FKConvexElem> NewConvex = MakeShared<FKConvexElem>();
+						if(NumPerformedSplits == 0)
+						{
+							// No more splits possible - decomp is complete
+							break;
+						}
 
-			if(bSimplifyHulls)
-			{
-				UE::Geometry::FDynamicMesh3 HullMesh = Decomp.GetHullMesh(HullIndex);
+						if(SplitIndex == MaxUnconstrainedSplits - 1)
+						{
+							UE_LOG(LogBoneGeometryGenerators, Warning, TEXT("Decomposition ran maximum number of splits (%d), potentially stuck in a splitting loop"), MaxUnconstrainedSplits);
+						}
+					}
 
-				UE::Geometry::FMeshConvexHull::FSimplifyOptions Options{
-					SimplifyTargetMaxFaces,
-					SimplifyGeometricToleranceAfterTarget
-				};
+					Decomp.FixHullOverlapsInNegativeSpace();
 
-				if(!ensure(UE::Geometry::FMeshConvexHull::SimplifyHull(HullMesh, Options)))
-				{
-					UE_LOG(LogBoneGeometryGenerators, Error, TEXT("Simplification failed for convex hull %d of %d, skipping"), HullIndex, NumComputedHulls);
-					continue;
-				}
+					UE::Geometry::FConvexDecomposition3::FMergeSettings MergeSettings;
+					MergeSettings.TargetNumParts = NegativeSpaceMaxSplits > 0 ? NegativeSpaceMaxSplits : -1;
+					MergeSettings.MinThicknessTolerance = MinThicknessTolerance;
+					MergeSettings.bAllowCompact = true;
 
-				for(int32 MeshVertIndex : HullMesh.VertexIndicesItr())
-				{
-					NewConvex->VertexData.Add(static_cast<FVector>(HullMesh.GetVertex(MeshVertIndex)));
+					Decomp.MergeBest(MergeSettings);
 				}
-			}
-			else
-			{
-				TArray<FVector3f> DecompVerts = Decomp.GetVertices<float>(HullIndex);
-				for(const FVector3f& Vert : DecompVerts)
+
+				const int32 NumComputedHulls = Decomp.NumHulls();
+				WorkerResults[BoneIndex].Reserve(NumComputedHulls);
+
+				for(int32 HullIndex = 0; HullIndex < NumComputedHulls; ++HullIndex)
 				{
-					NewConvex->VertexData.Add(static_cast<FVector>(Vert));
+					TSharedPtr<FKConvexElem> NewConvex = MakeShared<FKConvexElem>();
+
+					if(bSimplifyHulls)
+					{
+						UE::Geometry::FDynamicMesh3 HullMesh = Decomp.GetHullMesh(HullIndex);
+
+						UE::Geometry::FMeshConvexHull::FSimplifyOptions Options{
+							SimplifyTargetMaxFaces,
+							SimplifyGeometricToleranceAfterTarget
+						};
+
+						if(!ensure(UE::Geometry::FMeshConvexHull::SimplifyHull(HullMesh, Options)))
+						{
+							UE_LOG(LogBoneGeometryGenerators, Error, TEXT("Simplification failed for convex hull %d of %d, skipping"), HullIndex, NumComputedHulls);
+							continue;
+						}
+
+						for(int32 MeshVertIndex : HullMesh.VertexIndicesItr())
+						{
+							NewConvex->VertexData.Add(static_cast<FVector>(HullMesh.GetVertex(MeshVertIndex)));
+						}
+					}
+					else
+					{
+						TArray<FVector3f> DecompVerts = Decomp.GetVertices<float>(HullIndex);
+						for(const FVector3f& Vert : DecompVerts)
+						{
+							NewConvex->VertexData.Add(static_cast<FVector>(Vert));
+						}
+					}
+
+					NewConvex->UpdateElemBox();
+
+					WorkerResults[BoneIndex].Emplace(Bone, FSimpleGeometry(NewConvex));
 				}
 			}
+		});
+	UE::Tasks::Wait(DecompTasks);
 
-			NewConvex->UpdateElemBox();
+	TArray<FRigidAssetBoneGeometry> Result;
 
-			Result.Emplace(Bone, FSimpleGeometry(NewConvex));
-		}
+	const int32 TotalNumGeometry = Algo::TransformAccumulate(WorkerResults, [](const TArray<FRigidAssetBoneGeometry>& Item) { return Item.Num(); }, 0);
+	Result.Reserve(TotalNumGeometry);
+
+	for(const TArray<FRigidAssetBoneGeometry>& BoneGeometry : WorkerResults)
+	{
+		Result.Append(BoneGeometry);
 	}
 
 	return Result;
diff --git a/Engine/Plugins/Experimental/ChaosRigidAsset/Source/ChaosRigidAssetNodes/Private/PhysicsAssetDataflowNodes.cpp b/Engine/Plugins/Experimental/ChaosRigidAsset/Source/ChaosRigidAssetNodes/Private/PhysicsAssetDataflowNodes.cpp
--- a/Engine/Plugins/Experimental/ChaosRigidAsset/Source/ChaosRigidAssetNodes/Private/PhysicsAssetDataflowNodes.cpp
+++ b/Engine/Plugins/Experimental/ChaosRigidAsset/Source/ChaosRigidAssetNodes/Private/PhysicsAssetDataflowNodes.cpp
@@ -828,6 +828,9 @@ void FDataflowPhysicsAssetMakeState::Evaluate(UE::Dataflow::FContext& Context, c
 			}
 		}
 
+		// Make sure we have initialised the mesh in order to get vertex data
+		InTarget->CalculateInvRefMatrices();
+
 		SetValue(Context, FPhysicsAssetDataflowState{ InTarget->GetSkeleton(), InTarget }, &State);
 	}
 }
