Fixup docs on rigid caching node

#Synced-CL 52442783

diff --git a/Engine/Plugins/Experimental/ChaosRigidAsset/Source/ChaosRigidAssetNodes/Private/PhysicsAssetDataflowNodes.h b/Engine/Plugins/Experimental/ChaosRigidAsset/Source/ChaosRigidAssetNodes/Private/PhysicsAssetDataflowNodes.h
--- a/Engine/Plugins/Experimental/ChaosRigidAsset/Source/ChaosRigidAssetNodes/Private/PhysicsAssetDataflowNodes.h
+++ b/Engine/Plugins/Experimental/ChaosRigidAsset/Source/ChaosRigidAssetNodes/Private/PhysicsAssetDataflowNodes.h
@@ -581,11 +581,11 @@ public:
 
 private:
 
-	// Input physics asset state to use for the simulation
+	/** Input physics asset state to use for the simulation */
 	UPROPERTY(meta = (DataflowInput))
 	FPhysicsAssetDataflowState State;
 
-	/** OUtput animation asset  */
+	/** Output animation asset  */
 	UPROPERTY(meta = (DataflowOutput))
 	TObjectPtr<UAnimSequence> Animation = nullptr;
 
@@ -597,14 +597,15 @@ private:
 	UPROPERTY(EditAnywhere, Category = Caching, meta = (DataflowInput, ClampMin = 1))
 	int32 FrameRate = 30;
 
-	/** Gravity vector in cm/s2 */
+	/** Gravity vector in cm/s^2 */
 	UPROPERTY(EditAnywhere, Category = Simulation, meta = (DataflowInput))
 	FVector Gravity = FVector(0, 0, 980);
 
-	/** simulation space velocity in cm/s */
+	/** Simulation space velocity in cm/s */
 	UPROPERTY(EditAnywhere, Category = Simulation, meta = (DataflowInput))
 	FVector Velocity = FVector(0, 0, 0);
 
+	/** Affector used to manipulate objects in the simulation */
 	UPROPERTY(EditAnywhere, Category = Simulation, meta = (DataflowInput))
 	FDataflowSimulationAffectorTypes Affector;
 
