Rigid asset - Update rigid asset asset to use the unified dataflow menu command so that the user experience is the same across all assets

#Synced-CL 52497132

diff --git a/Engine/Plugins/Experimental/ChaosRigidAsset/Source/ChaosRigidAssetEditor/Private/ChaosRigidAssetEditorModule.cpp b/Engine/Plugins/Experimental/ChaosRigidAsset/Source/ChaosRigidAssetEditor/Private/ChaosRigidAssetEditorModule.cpp
--- a/Engine/Plugins/Experimental/ChaosRigidAsset/Source/ChaosRigidAssetEditor/Private/ChaosRigidAssetEditorModule.cpp
+++ b/Engine/Plugins/Experimental/ChaosRigidAsset/Source/ChaosRigidAssetEditor/Private/ChaosRigidAssetEditorModule.cpp
@@ -35,15 +35,6 @@ public:
 	bool OpenAsset(UPhysicsAsset* InAsset) override;
 };
 
-UPhysicsAsset* GetPhysAsset(const FAssetData& InAsset)
-{
-	if(UObject* AssetBase = InAsset.GetAsset())
-	{
-		return Cast<UPhysicsAsset>(AssetBase);
-	}
-	return nullptr;
-}
-
 UPhysicsAssetDataflowAttachment* GetDataflowAttachment(IInterface_AssetUserData* UserDataHolder)
 {
 	if(!UserDataHolder)
@@ -54,40 +45,6 @@ UPhysicsAssetDataflowAttachment* GetDataflowAttachment(IInterface_AssetUserData*
 	return UserDataHolder->GetAssetUserData<UPhysicsAssetDataflowAttachment>();
 }
 
-bool HasDataflowAttachment(IInterface_AssetUserData* UserDataHolder)
-{
-	if(!UserDataHolder)
-	{
-		return false;
-	}
-
-	return UserDataHolder->HasAssetUserDataOfClass(UPhysicsAssetDataflowAttachment::StaticClass());
-}
-
-void AddDataflow(FAssetData InAsset)
-{
-	UPhysicsAsset* Asset = GetPhysAsset(InAsset);
-
-	if(!Asset)
-	{
-		return;
-	}
-
-	Asset->AddAssetUserData(NewObject<UPhysicsAssetDataflowAttachment>(Asset));
-}
-
-void RemoveDataflow(FAssetData InAsset)
-{
-	UPhysicsAsset* Asset = GetPhysAsset(InAsset);
-
-	if(!Asset)
-	{
-		return;
-	}
-
-	Asset->RemoveUserDataOfClass(UPhysicsAssetDataflowAttachment::StaticClass());
-}
-
 void SpawnEditorFor(UPhysicsAsset* InAsset)
 {
 	UDataflowAttachment* Attachment = GetDataflowAttachment(InAsset);
@@ -118,91 +75,8 @@ void SpawnEditorFor(UPhysicsAsset* InAsset)
 	}
 }
 
-void EditDataflow(FAssetData InAsset)
-{
-	SpawnEditorFor(GetPhysAsset(InAsset));
-}
-
-bool CanAddDataflow(FAssetData InAsset)
-{
-	UPhysicsAsset* Asset = GetPhysAsset(InAsset);
-	return Asset && !HasDataflowAttachment(Asset);
-}
-
-bool CanRemoveDataflow(FAssetData InAsset)
-{
-	UPhysicsAsset* Asset = GetPhysAsset(InAsset);
-	return Asset && HasDataflowAttachment(Asset);
-}
-
-bool CanEditDataflow(FAssetData InAsset)
-{
-	UPhysicsAsset* Asset = GetPhysAsset(InAsset);
-	return Asset && HasDataflowAttachment(Asset);
-}
-
-void FillMenu(FMenuBuilder& MenuBuilder, const TArray<FAssetData> Assets)
-{
-	if(Assets.Num() == 1 && Assets[0].IsInstanceOf<UPhysicsAsset>())
-	{
-		UPhysicsAsset* PhysAsset = GetPhysAsset(Assets[0]);
-
-		if(!PhysAsset)
-		{
-			return;
-		}
-
-		const bool bHasAttachment = HasDataflowAttachment(PhysAsset);
-
-		MenuBuilder.BeginSection("PhysicsAssetMenu_Dataflow", NSLOCTEXT("RigidAsset", "MenuSection", "DataFlow"));
-
-		// Add Dataflow
-		{
-			FMenuEntryParams Params;
-			Params.DirectActions.ExecuteAction = FExecuteAction::CreateStatic(&AddDataflow, Assets[0]);
-			Params.DirectActions.CanExecuteAction = FCanExecuteAction::CreateStatic(&CanAddDataflow, Assets[0]);
-			Params.LabelOverride = NSLOCTEXT("ChaosRigidAsset", "PhysicsAssetContextAddDataflow", "Add Dataflow");
-			MenuBuilder.AddMenuEntry(Params);
-		}
-
-		// Remove Dataflow
-		{
-			FMenuEntryParams Params;
-			Params.DirectActions.ExecuteAction = FExecuteAction::CreateStatic(&RemoveDataflow, Assets[0]);
-			Params.DirectActions.CanExecuteAction = FCanExecuteAction::CreateStatic(&CanRemoveDataflow, Assets[0]);
-			Params.LabelOverride = NSLOCTEXT("ChaosRigidAsset", "PhysicsAssetContextRemoveDataflow", "Remove Dataflow");
-			MenuBuilder.AddMenuEntry(Params);
-		}
-
-		// Edit Dataflow
-		{
-			FMenuEntryParams Params;
-			Params.DirectActions.ExecuteAction = FExecuteAction::CreateStatic(&EditDataflow, Assets[0]);
-			Params.DirectActions.CanExecuteAction = FCanExecuteAction::CreateStatic(&CanEditDataflow, Assets[0]);
-			Params.LabelOverride = NSLOCTEXT("ChaosRigidAsset", "PhysicsAssetContextEditDataflow", "Edit Dataflow");
-			MenuBuilder.AddMenuEntry(Params);
-		}
-
-		MenuBuilder.EndSection();
-	}
-}
-
-TSharedRef<FExtender> ExtendAssetContextMenu(const TArray<FAssetData>& Assets)
-{
-	TSharedRef<FExtender> Extender = MakeShared<FExtender>();
-
-	Extender->AddMenuExtension("GetAssetActions", EExtensionHook::After, nullptr, FMenuExtensionDelegate::CreateStatic(&FillMenu, Assets));
-
-	return Extender;
-}
-
 void FChaosRigidAssetEditorModule::StartupModule()
 {
-	FContentBrowserModule& ContentBrowser = FModuleManager::LoadModuleChecked<FContentBrowserModule>(TEXT("ContentBrowser"));
-	TArray<FContentBrowserMenuExtender_SelectedAssets>& AssetMenuExtenders = ContentBrowser.GetAllAssetViewContextMenuExtenders();
-	AssetMenuExtenders.Add(FContentBrowserMenuExtender_SelectedAssets::CreateStatic(&ExtendAssetContextMenu));
-	AssetMenuExtenderHandle = AssetMenuExtenders.Last().GetHandle();
-
 	EditorFeature = MakeUnique<FDataflowPhysicsAssetEditorOverride>();
 
 	IModularFeatures::Get().RegisterModularFeature(IPhysicsAssetEditorOverride::ModularFeatureName, EditorFeature.Get());
@@ -210,6 +84,9 @@ void FChaosRigidAssetEditorModule::StartupModule()
 	UE::Dataflow::FRenderingFactory::GetInstance()->RegisterCallbacks(MakeUnique<UE::Chaos::RigidAsset::FAggregateGeometryGeomRenderCallbacks>());
 	UE::Dataflow::FRenderingFactory::GetInstance()->RegisterCallbacks(MakeUnique<UE::Chaos::RigidAsset::FBoneSelectionRenderCallbacks>());
 	UE::Dataflow::FRenderingFactory::GetInstance()->RegisterCallbacks(MakeUnique<UE::Chaos::RigidAsset::FPhysAssetStateRenderCallbacks>());
+
+	FDataflowAttachmentFactory::Get().Register(UPhysicsAsset::StaticClass()->GetFName(), [](UObject* Owner) { return NewObject<UPhysicsAssetDataflowAttachment>(Owner); });
+	DataflowMenusHandle = UE::DataflowAssetDefinitionHelpers::RegisterDataflowAssetMenus(UPhysicsAsset::StaticClass());
 }
 
 void FChaosRigidAssetEditorModule::ShutdownModule()
@@ -217,22 +94,11 @@ void FChaosRigidAssetEditorModule::ShutdownModule()
 	IModularFeatures::Get().UnregisterModularFeature(IPhysicsAssetEditorOverride::ModularFeatureName, EditorFeature.Get());
 	EditorFeature = nullptr;
 
-	if(AssetMenuExtenderHandle.IsValid() && FModuleManager::Get().IsModuleLoaded(TEXT("ContentBrowser")))
-	{
-		FContentBrowserModule& ContentBrowserModule = FModuleManager::LoadModuleChecked<FContentBrowserModule>(TEXT("ContentBrowser"));
-		TArray<FContentBrowserMenuExtender_SelectedAssets>& AssetMenuExtenders = ContentBrowserModule.GetAllAssetViewContextMenuExtenders();
-
-		AssetMenuExtenders.RemoveAll([this](const FContentBrowserMenuExtender_SelectedAssets& InDelegate)
-			{
-				return InDelegate.GetHandle() == AssetMenuExtenderHandle;
-			});
-
-		AssetMenuExtenderHandle.Reset();
-	}
-
 	UE::Dataflow::FRenderingFactory::GetInstance()->DeregisterCallbacks(UE::Chaos::RigidAsset::FPhysAssetStateRenderCallbacks::StaticGetRenderKey());
 	UE::Dataflow::FRenderingFactory::GetInstance()->DeregisterCallbacks(UE::Chaos::RigidAsset::FBoneSelectionRenderCallbacks::StaticGetRenderKey());
 	UE::Dataflow::FRenderingFactory::GetInstance()->DeregisterCallbacks(UE::Chaos::RigidAsset::FAggregateGeometryGeomRenderCallbacks::StaticGetRenderKey());
+
+	UE::DataflowAssetDefinitionHelpers::UnregisterDataflowAssetMenus(DataflowMenusHandle);
 }
 
 bool FDataflowPhysicsAssetEditorOverride::OpenAsset(UPhysicsAsset* InAsset)
diff --git a/Engine/Plugins/Experimental/ChaosRigidAsset/Source/ChaosRigidAssetEditor/Public/ChaosRigidAssetEditorModule.h b/Engine/Plugins/Experimental/ChaosRigidAsset/Source/ChaosRigidAssetEditor/Public/ChaosRigidAssetEditorModule.h
--- a/Engine/Plugins/Experimental/ChaosRigidAsset/Source/ChaosRigidAssetEditor/Public/ChaosRigidAssetEditorModule.h
+++ b/Engine/Plugins/Experimental/ChaosRigidAsset/Source/ChaosRigidAssetEditor/Public/ChaosRigidAssetEditorModule.h
@@ -18,4 +18,5 @@ public:
 private:
 	TUniquePtr<FDataflowPhysicsAssetEditorOverride> EditorFeature;
 	FDelegateHandle AssetMenuExtenderHandle;
+	FDelegateHandle DataflowMenusHandle;
 };
